Fixed MacTargetPlatform compile errors

diff --git a/Engine/Source/Developer/Mac/MacNoEditorTargetPlatform/Private/MacNoEditorTargetPlatformModule.cpp b/Engine/Source/Developer/Mac/MacNoEditorTargetPlatform/Private/MacNoEditorTargetPlatformModule.cpp
--- a/Engine/Source/Developer/Mac/MacNoEditorTargetPlatform/Private/MacNoEditorTargetPlatformModule.cpp
+++ b/Engine/Source/Developer/Mac/MacNoEditorTargetPlatform/Private/MacNoEditorTargetPlatformModule.cpp
@@ -36,7 +36,7 @@ public:
 		Singleton = NULL;
 	}
 
-	virtual ITargetPlatform* GetTargetPlatform( )
+	virtual ITargetPlatform* GetTargetPlatform( ) OVERRIDE
 	{
 		if (Singleton == NULL)
 		{
diff --git a/Engine/Source/Developer/Mac/MacTargetPlatform/Private/GenericMacTargetPlatform.h b/Engine/Source/Developer/Mac/MacTargetPlatform/Private/GenericMacTargetPlatform.h
--- a/Engine/Source/Developer/Mac/MacTargetPlatform/Private/GenericMacTargetPlatform.h
+++ b/Engine/Source/Developer/Mac/MacTargetPlatform/Private/GenericMacTargetPlatform.h
@@ -34,6 +34,8 @@ public:
 
 	// Begin ITargetPlatform interface
 
+	virtual void EnableDeviceCheck(bool OnOff) OVERRIDE {}
+
 	virtual void GetAllDevices( TArray<ITargetDevicePtr>& OutDevices ) const OVERRIDE
 	{
 		OutDevices.Reset();
diff --git a/Engine/Source/Developer/Mac/MacTargetPlatform/Private/MacTargetPlatformModule.cpp b/Engine/Source/Developer/Mac/MacTargetPlatform/Private/MacTargetPlatformModule.cpp
--- a/Engine/Source/Developer/Mac/MacTargetPlatform/Private/MacTargetPlatformModule.cpp
+++ b/Engine/Source/Developer/Mac/MacTargetPlatform/Private/MacTargetPlatformModule.cpp
@@ -49,7 +49,7 @@ public:
 
 	// Begin ITargetPlatformModule interface
 
-	virtual ITargetPlatform* GetTargetPlatform()
+	virtual ITargetPlatform* GetTargetPlatform() OVERRIDE
 	{
 		if (Singleton == NULL)
 		{
