Added MacClient and MacServer target platforms

diff --git a/Engine/Source/Developer/Mac/MacClientTargetPlatform/MacClientTargetPlatform.Build.cs b/Engine/Source/Developer/Mac/MacClientTargetPlatform/MacClientTargetPlatform.Build.cs
new file mode 100644
--- /dev/null
+++ b/Engine/Source/Developer/Mac/MacClientTargetPlatform/MacClientTargetPlatform.Build.cs
@@ -0,0 +1,29 @@
+// Copyright 1998-2014 Epic Games, Inc. All Rights Reserved.
+
+using UnrealBuildTool;
+
+public class MacClientTargetPlatform : ModuleRules
+{
+	public MacClientTargetPlatform(TargetInfo Target)
+	{
+		PrivateDependencyModuleNames.AddRange(
+			new string[] {
+				"Core",
+				"TargetPlatform"
+			}
+		);
+
+		if (UEBuildConfiguration.bCompileAgainstEngine)
+		{
+			PrivateDependencyModuleNames.Add("CoreUObject"); // @todo Mac: for some reason it's needed to link in debug on Mac
+			PrivateDependencyModuleNames.Add("Engine");
+			PrivateIncludePathModuleNames.Add("TextureCompressor");
+		}
+
+		PrivateIncludePaths.AddRange(
+			new string[] {
+				"Developer/Mac/MacTargetPlatform/Private"
+			}
+		);
+	}
+}
diff --git a/Engine/Source/Developer/Mac/MacClientTargetPlatform/Private/MacClientTargetPlatformModule.cpp b/Engine/Source/Developer/Mac/MacClientTargetPlatform/Private/MacClientTargetPlatformModule.cpp
new file mode 100644
--- /dev/null
+++ b/Engine/Source/Developer/Mac/MacClientTargetPlatform/Private/MacClientTargetPlatformModule.cpp
@@ -0,0 +1,51 @@
+// Copyright 1998-2014 Epic Games, Inc. All Rights Reserved.
+
+/*=============================================================================
+	MacClientTargetPlatformModule.cpp: Implements the FMacClientTargetPlatformModule class.
+=============================================================================*/
+
+#if PLATFORM_MAC
+#define FVector FVectorWorkaround
+#ifdef __OBJC__
+#import <Cocoa/Cocoa.h>
+#endif
+#undef check
+#undef verify
+#undef FVector
+#endif
+
+#include "MacClientTargetPlatformPrivatePCH.h"
+
+
+/**
+ * Holds the target platform singleton.
+ */
+static ITargetPlatform* Singleton = NULL;
+
+
+/**
+ * Module for the Mac target platform (without editor).
+ */
+class FMacClientTargetPlatformModule
+	: public ITargetPlatformModule
+{
+public:
+
+	virtual ~FMacClientTargetPlatformModule( )
+	{
+		Singleton = NULL;
+	}
+
+	virtual ITargetPlatform* GetTargetPlatform( ) OVERRIDE
+	{
+		if (Singleton == NULL)
+		{
+			Singleton = new TGenericMacTargetPlatform<false, false, true>();
+		}
+
+		return Singleton;
+	}
+};
+
+
+IMPLEMENT_MODULE(FMacClientTargetPlatformModule, MacClientTargetPlatform);
\ No newline at end of file
diff --git a/Engine/Source/Developer/Mac/MacClientTargetPlatform/Private/MacClientTargetPlatformPrivatePCH.h b/Engine/Source/Developer/Mac/MacClientTargetPlatform/Private/MacClientTargetPlatformPrivatePCH.h
new file mode 100644
--- /dev/null
+++ b/Engine/Source/Developer/Mac/MacClientTargetPlatform/Private/MacClientTargetPlatformPrivatePCH.h
@@ -0,0 +1,27 @@
+// Copyright 1998-2014 Epic Games, Inc. All Rights Reserved.
+
+/*=============================================================================
+	MacClientTargetPlatformPrivatePCH.h: Pre-compiled header file for the MacClientTargetPlatform module.
+=============================================================================*/
+
+#pragma once
+
+
+/* Dependencies
+ *****************************************************************************/
+
+#include "Core.h"
+#include "ModuleManager.h"
+#include "MacPlatformProperties.h"
+
+#if WITH_ENGINE
+	#include "Engine.h"
+	#include "TextureCompressorModule.h"
+#endif
+
+#include "TargetPlatform.h"
+#include "LocalMacTargetDevice.h"
+#include "GenericMacTargetPlatform.h"
+
+/* Private includes
+ *****************************************************************************/
\ No newline at end of file
diff --git a/Engine/Source/Developer/Mac/MacNoEditorTargetPlatform/Private/MacNoEditorTargetPlatformModule.cpp b/Engine/Source/Developer/Mac/MacNoEditorTargetPlatform/Private/MacNoEditorTargetPlatformModule.cpp
--- a/Engine/Source/Developer/Mac/MacNoEditorTargetPlatform/Private/MacNoEditorTargetPlatformModule.cpp
+++ b/Engine/Source/Developer/Mac/MacNoEditorTargetPlatform/Private/MacNoEditorTargetPlatformModule.cpp
@@ -40,7 +40,7 @@ public:
 	{
 		if (Singleton == NULL)
 		{
-			Singleton = new TGenericMacTargetPlatform<false, false>();
+			Singleton = new TGenericMacTargetPlatform<false, false, false>();
 		}
 
 		return Singleton;
diff --git a/Engine/Source/Developer/Mac/MacServerTargetPlatform/MacServerTargetPlatform.Build.cs b/Engine/Source/Developer/Mac/MacServerTargetPlatform/MacServerTargetPlatform.Build.cs
new file mode 100644
--- /dev/null
+++ b/Engine/Source/Developer/Mac/MacServerTargetPlatform/MacServerTargetPlatform.Build.cs
@@ -0,0 +1,29 @@
+// Copyright 1998-2014 Epic Games, Inc. All Rights Reserved.
+
+using UnrealBuildTool;
+
+public class MacServerTargetPlatform : ModuleRules
+{
+	public MacServerTargetPlatform(TargetInfo Target)
+	{
+		PrivateDependencyModuleNames.AddRange(
+			new string[] {
+				"Core",
+				"TargetPlatform"
+			}
+		);
+
+		if (UEBuildConfiguration.bCompileAgainstEngine)
+		{
+			PrivateDependencyModuleNames.Add("CoreUObject"); // @todo Mac: for some reason it's needed to link in debug on Mac
+			PrivateDependencyModuleNames.Add("Engine");
+			PrivateIncludePathModuleNames.Add("TextureCompressor");
+		}
+
+		PrivateIncludePaths.AddRange(
+			new string[] {
+				"Developer/Mac/MacTargetPlatform/Private"
+			}
+		);
+	}
+}
diff --git a/Engine/Source/Developer/Mac/MacServerTargetPlatform/Private/MacServerTargetPlatformModule.cpp b/Engine/Source/Developer/Mac/MacServerTargetPlatform/Private/MacServerTargetPlatformModule.cpp
new file mode 100644
--- /dev/null
+++ b/Engine/Source/Developer/Mac/MacServerTargetPlatform/Private/MacServerTargetPlatformModule.cpp
@@ -0,0 +1,51 @@
+// Copyright 1998-2014 Epic Games, Inc. All Rights Reserved.
+
+/*=============================================================================
+	MacServerTargetPlatformModule.cpp: Implements the FMacServerTargetPlatformModule class.
+=============================================================================*/
+
+#if PLATFORM_MAC
+#define FVector FVectorWorkaround
+#ifdef __OBJC__
+#import <Cocoa/Cocoa.h>
+#endif
+#undef check
+#undef verify
+#undef FVector
+#endif
+
+#include "MacServerTargetPlatformPrivatePCH.h"
+
+
+/**
+ * Holds the target platform singleton.
+ */
+static ITargetPlatform* Singleton = NULL;
+
+
+/**
+ * Module for the Mac target platform (without editor).
+ */
+class FMacServerTargetPlatformModule
+	: public ITargetPlatformModule
+{
+public:
+
+	virtual ~FMacServerTargetPlatformModule( )
+	{
+		Singleton = NULL;
+	}
+
+	virtual ITargetPlatform* GetTargetPlatform( ) OVERRIDE
+	{
+		if (Singleton == NULL)
+		{
+			Singleton = new TGenericMacTargetPlatform<false, true, false>();
+		}
+
+		return Singleton;
+	}
+};
+
+
+IMPLEMENT_MODULE(FMacServerTargetPlatformModule, MacServerTargetPlatform);
\ No newline at end of file
diff --git a/Engine/Source/Developer/Mac/MacServerTargetPlatform/Private/MacServerTargetPlatformPrivatePCH.h b/Engine/Source/Developer/Mac/MacServerTargetPlatform/Private/MacServerTargetPlatformPrivatePCH.h
new file mode 100644
--- /dev/null
+++ b/Engine/Source/Developer/Mac/MacServerTargetPlatform/Private/MacServerTargetPlatformPrivatePCH.h
@@ -0,0 +1,27 @@
+// Copyright 1998-2014 Epic Games, Inc. All Rights Reserved.
+
+/*=============================================================================
+	MacServerTargetPlatformPrivatePCH.h: Pre-compiled header file for the MacServerTargetPlatform module.
+=============================================================================*/
+
+#pragma once
+
+
+/* Dependencies
+ *****************************************************************************/
+
+#include "Core.h"
+#include "ModuleManager.h"
+#include "MacPlatformProperties.h"
+
+#if WITH_ENGINE
+	#include "Engine.h"
+	#include "TextureCompressorModule.h"
+#endif
+
+#include "TargetPlatform.h"
+#include "LocalMacTargetDevice.h"
+#include "GenericMacTargetPlatform.h"
+
+/* Private includes
+ *****************************************************************************/
\ No newline at end of file
diff --git a/Engine/Source/Developer/Mac/MacTargetPlatform/Private/GenericMacTargetPlatform.h b/Engine/Source/Developer/Mac/MacTargetPlatform/Private/GenericMacTargetPlatform.h
--- a/Engine/Source/Developer/Mac/MacTargetPlatform/Private/GenericMacTargetPlatform.h
+++ b/Engine/Source/Developer/Mac/MacTargetPlatform/Private/GenericMacTargetPlatform.h
@@ -10,12 +10,14 @@
 /**
  * Template for Mac target platforms
  */
-template<bool HAS_EDITOR_DATA, bool IS_DEDICATED_SERVER>
+template<bool HAS_EDITOR_DATA, bool IS_DEDICATED_SERVER, bool IS_CLIENT_ONLY>
 class TGenericMacTargetPlatform
-	: public TTargetPlatformBase<FMacPlatformProperties<HAS_EDITOR_DATA, IS_DEDICATED_SERVER> >
+	: public TTargetPlatformBase<FMacPlatformProperties<HAS_EDITOR_DATA, IS_DEDICATED_SERVER, IS_CLIENT_ONLY> >
 {
 public:
 
+	typedef TTargetPlatformBase<FMacPlatformProperties<HAS_EDITOR_DATA, IS_DEDICATED_SERVER, IS_CLIENT_ONLY> > TSuper;
+
 	/**
 	 * Default constructor.
 	 */
@@ -127,7 +129,7 @@ public:
 			return (HAS_EDITOR_DATA || !IS_DEDICATED_SERVER);
 		}
 
-		return TTargetPlatformBase<FMacPlatformProperties<HAS_EDITOR_DATA, IS_DEDICATED_SERVER> >::SupportsFeature(Feature);
+return TSuper::SupportsFeature(Feature);
 	}
 
 #if WITH_ENGINE
diff --git a/Engine/Source/Developer/Mac/MacTargetPlatform/Private/MacTargetPlatformModule.cpp b/Engine/Source/Developer/Mac/MacTargetPlatform/Private/MacTargetPlatformModule.cpp
--- a/Engine/Source/Developer/Mac/MacTargetPlatform/Private/MacTargetPlatformModule.cpp
+++ b/Engine/Source/Developer/Mac/MacTargetPlatform/Private/MacTargetPlatformModule.cpp
@@ -53,7 +53,7 @@ public:
 	{
 		if (Singleton == NULL)
 		{
-			Singleton = new TGenericMacTargetPlatform<true, false>;
+			Singleton = new TGenericMacTargetPlatform<true, false, false>();
 		}
 
 		return Singleton;
diff --git a/Engine/Source/Programs/UnrealBuildTool/Mac/UEBuildMac.cs b/Engine/Source/Programs/UnrealBuildTool/Mac/UEBuildMac.cs
--- a/Engine/Source/Programs/UnrealBuildTool/Mac/UEBuildMac.cs
+++ b/Engine/Source/Programs/UnrealBuildTool/Mac/UEBuildMac.cs
@@ -97,13 +97,15 @@ namespace UnrealBuildTool
                     }
                 }
 
-		// allow standalone tools to use target platform modules, without needing Engine
-		if (UEBuildConfiguration.bForceBuildTargetPlatforms)
-		{
-			InModule.AddPlatformSpecificDynamicallyLoadedModule("HTML5TargetPlatform");
-			InModule.AddDynamicallyLoadedModule("MacTargetPlatform");
-			InModule.AddDynamicallyLoadedModule("MacNoEditorTargetPlatform");
-		}
+				// allow standalone tools to use target platform modules, without needing Engine
+				if (UEBuildConfiguration.bForceBuildTargetPlatforms)
+				{
+					InModule.AddPlatformSpecificDynamicallyLoadedModule("HTML5TargetPlatform");
+					InModule.AddDynamicallyLoadedModule("MacTargetPlatform");
+					InModule.AddDynamicallyLoadedModule("MacNoEditorTargetPlatform");
+					InModule.AddDynamicallyLoadedModule("MacClientTargetPlatform");
+					InModule.AddDynamicallyLoadedModule("MacServerTargetPlatform");
+				}
 
                 if (bBuildShaderFormats)
                 {
diff --git a/Engine/Source/Runtime/Core/Public/Mac/MacPlatformIncludes.h b/Engine/Source/Runtime/Core/Public/Mac/MacPlatformIncludes.h
--- a/Engine/Source/Runtime/Core/Public/Mac/MacPlatformIncludes.h
+++ b/Engine/Source/Runtime/Core/Public/Mac/MacPlatformIncludes.h
@@ -104,5 +104,4 @@ struct RECT
 #error "WITH_EDITORONLY_DATA must be defined"
 #endif
 
-typedef FMacPlatformProperties<WITH_EDITORONLY_DATA, UE_SERVER> FPlatformProperties;
-
+typedef FMacPlatformProperties<WITH_EDITORONLY_DATA, UE_SERVER, !WITH_SERVER_CODE> FPlatformProperties;
diff --git a/Engine/Source/Runtime/Core/Public/Mac/MacPlatformProperties.h b/Engine/Source/Runtime/Core/Public/Mac/MacPlatformProperties.h
--- a/Engine/Source/Runtime/Core/Public/Mac/MacPlatformProperties.h
+++ b/Engine/Source/Runtime/Core/Public/Mac/MacPlatformProperties.h
@@ -15,7 +15,7 @@
 /**
  * Implements Mac platform properties.
  */
-template<bool HAS_EDITOR_DATA, bool IS_DEDICATED_SERVER>
+template<bool HAS_EDITOR_DATA, bool IS_DEDICATED_SERVER, bool IS_CLIENT_ONLY>
 struct FMacPlatformProperties
 	: public FGenericPlatformProperties
 {
@@ -31,6 +31,11 @@ struct FMacPlatformProperties
 			return "Mac (Editor)";
 		}
 
+		if (IS_CLIENT_ONLY)
+		{
+			return "Mac (Client-only)";
+		}
+
 		return "Mac";
 	}
 
@@ -54,6 +59,11 @@ struct FMacPlatformProperties
 		return IS_DEDICATED_SERVER;
 	}
 
+	static FORCEINLINE bool IsClientOnly()
+	{
+		return IS_CLIENT_ONLY;
+	}
+
 	static FORCEINLINE const char* PlatformName( )
 	{
 		if (IS_DEDICATED_SERVER)
@@ -66,6 +76,11 @@ struct FMacPlatformProperties
 			return "Mac";
 		}
 
+		if (IS_CLIENT_ONLY)
+		{
+			return "MacClient";
+		}
+
 		return "MacNoEditor";
 	}
 
diff --git a/Engine/Source/Runtime/Engine/Engine.Build.cs b/Engine/Source/Runtime/Engine/Engine.Build.cs
--- a/Engine/Source/Runtime/Engine/Engine.Build.cs
+++ b/Engine/Source/Runtime/Engine/Engine.Build.cs
@@ -115,6 +115,8 @@ public class Engine : ModuleRules
 					new string[] {
 					    "MacTargetPlatform",
 					    "MacNoEditorTargetPlatform",
+						"MacServerTargetPlatform",
+						"MacClientTargetPlatform",
 					    "HTML5TargetPlatform"
 					}
 				);
