CIS fix to mac platform

diff --git a/Engine/Source/Developer/Mac/MacTargetPlatform/Private/GenericMacTargetPlatform.h b/Engine/Source/Developer/Mac/MacTargetPlatform/Private/GenericMacTargetPlatform.h
--- a/Engine/Source/Developer/Mac/MacTargetPlatform/Private/GenericMacTargetPlatform.h
+++ b/Engine/Source/Developer/Mac/MacTargetPlatform/Private/GenericMacTargetPlatform.h
@@ -6,6 +6,9 @@
 
 #pragma once
 
+#if WITH_ENGINE
+#include "StaticMeshResources.h"
+#endif // WITH_ENGINE
 
 /**
  * Template for Mac target platforms
