Small cleanup of Mac*TargetPlatform.cpp files

diff --git a/Engine/Source/Developer/Mac/MacClientTargetPlatform/Private/MacClientTargetPlatformModule.cpp b/Engine/Source/Developer/Mac/MacClientTargetPlatform/Private/MacClientTargetPlatformModule.cpp
--- a/Engine/Source/Developer/Mac/MacClientTargetPlatform/Private/MacClientTargetPlatformModule.cpp
+++ b/Engine/Source/Developer/Mac/MacClientTargetPlatform/Private/MacClientTargetPlatformModule.cpp
@@ -4,16 +4,6 @@
 	MacClientTargetPlatformModule.cpp: Implements the FMacClientTargetPlatformModule class.
 =============================================================================*/
 
-#if PLATFORM_MAC
-#define FVector FVectorWorkaround
-#ifdef __OBJC__
-#import <Cocoa/Cocoa.h>
-#endif
-#undef check
-#undef verify
-#undef FVector
-#endif
-
 #include "MacClientTargetPlatformPrivatePCH.h"
 
 
diff --git a/Engine/Source/Developer/Mac/MacNoEditorTargetPlatform/Private/MacNoEditorTargetPlatformModule.cpp b/Engine/Source/Developer/Mac/MacNoEditorTargetPlatform/Private/MacNoEditorTargetPlatformModule.cpp
--- a/Engine/Source/Developer/Mac/MacNoEditorTargetPlatform/Private/MacNoEditorTargetPlatformModule.cpp
+++ b/Engine/Source/Developer/Mac/MacNoEditorTargetPlatform/Private/MacNoEditorTargetPlatformModule.cpp
@@ -4,16 +4,6 @@
 	MacNoEditorTargetPlatformModule.cpp: Implements the FMacNoEditorTargetPlatformModule class.
 =============================================================================*/
 
-#if PLATFORM_MAC
-#define FVector FVectorWorkaround
-#ifdef __OBJC__
-#import <Cocoa/Cocoa.h>
-#endif
-#undef check
-#undef verify
-#undef FVector
-#endif
-
 #include "MacNoEditorTargetPlatformPrivatePCH.h"
 
 
diff --git a/Engine/Source/Developer/Mac/MacServerTargetPlatform/Private/MacServerTargetPlatformModule.cpp b/Engine/Source/Developer/Mac/MacServerTargetPlatform/Private/MacServerTargetPlatformModule.cpp
--- a/Engine/Source/Developer/Mac/MacServerTargetPlatform/Private/MacServerTargetPlatformModule.cpp
+++ b/Engine/Source/Developer/Mac/MacServerTargetPlatform/Private/MacServerTargetPlatformModule.cpp
@@ -4,16 +4,6 @@
 	MacServerTargetPlatformModule.cpp: Implements the FMacServerTargetPlatformModule class.
 =============================================================================*/
 
-#if PLATFORM_MAC
-#define FVector FVectorWorkaround
-#ifdef __OBJC__
-#import <Cocoa/Cocoa.h>
-#endif
-#undef check
-#undef verify
-#undef FVector
-#endif
-
 #include "MacServerTargetPlatformPrivatePCH.h"
 
 
diff --git a/Engine/Source/Developer/Mac/MacTargetPlatform/Private/MacTargetPlatformModule.cpp b/Engine/Source/Developer/Mac/MacTargetPlatform/Private/MacTargetPlatformModule.cpp
--- a/Engine/Source/Developer/Mac/MacTargetPlatform/Private/MacTargetPlatformModule.cpp
+++ b/Engine/Source/Developer/Mac/MacTargetPlatform/Private/MacTargetPlatformModule.cpp
@@ -4,18 +4,6 @@
 	MacTargetPlatformModule.cpp: Implements the FMacTargetPlatformModule class.
 =============================================================================*/
 
-#if PLATFORM_MAC
-	#define FVector FVectorWorkaround
-
-	#ifdef __OBJC__
-		#import <Cocoa/Cocoa.h>
-	#endif
-
-	#undef check
-	#undef verify
-	#undef FVector
-#endif
-
 #include "MacTargetPlatformPrivatePCH.h"
 
 
