MacOS: Added missing header include

diff --git a/Engine/Source/Developer/Mac/MacTargetPlatform/Private/MacTargetPlatformPrivatePCH.h b/Engine/Source/Developer/Mac/MacTargetPlatform/Private/MacTargetPlatformPrivatePCH.h
--- a/Engine/Source/Developer/Mac/MacTargetPlatform/Private/MacTargetPlatformPrivatePCH.h
+++ b/Engine/Source/Developer/Mac/MacTargetPlatform/Private/MacTargetPlatformPrivatePCH.h
@@ -3,10 +3,11 @@
 #pragma once
 
 
-/* Dependencies
+/* Private dependencies
  *****************************************************************************/
 
 #include "Core.h"
+#include "CoreUObject.h"
 #include "Runtime/Core/Public/Mac/MacPlatformProperties.h"
 
 #if WITH_ENGINE
@@ -23,4 +24,4 @@
 
 #include "MacTargetSettings.h"
 #include "LocalMacTargetDevice.h"
-#include "GenericMacTargetPlatform.h"
\ No newline at end of file
+#include "GenericMacTargetPlatform.h"
