Avoid trying to use unsupported BC4 texture compression on OS X.

diff --git a/Engine/Source/Developer/Mac/MacTargetPlatform/Private/GenericMacTargetPlatform.h b/Engine/Source/Developer/Mac/MacTargetPlatform/Private/GenericMacTargetPlatform.h
--- a/Engine/Source/Developer/Mac/MacTargetPlatform/Private/GenericMacTargetPlatform.h
+++ b/Engine/Source/Developer/Mac/MacTargetPlatform/Private/GenericMacTargetPlatform.h
@@ -122,6 +122,12 @@ return TSuper::SupportsFeature(Feature);
 			// no textures needed for dedicated server target
 			OutFormats.Add(NAME_None);
 		}
+		else if ( InTexture->CompressionSettings == TC_Alpha)
+		{
+			// no BC4 texture compression on OS X
+			static FName NameG8(TEXT("G8"));
+			OutFormats.Add(NameG8);
+		}
 		else
 		{
 			// just use the standard texture format name for this texture
