Update MacTargetPlatform code to remove a deprecation warning.

diff --git a/Engine/Source/Developer/Mac/MacTargetPlatform/Private/MacTargetPlatformModule.cpp b/Engine/Source/Developer/Mac/MacTargetPlatform/Private/MacTargetPlatformModule.cpp
--- a/Engine/Source/Developer/Mac/MacTargetPlatform/Private/MacTargetPlatformModule.cpp
+++ b/Engine/Source/Developer/Mac/MacTargetPlatform/Private/MacTargetPlatformModule.cpp
@@ -54,7 +54,7 @@ public:
 
 	virtual void StartupModule() override
 	{
-		TargetSettings = ConstructObject<UMacTargetSettings>(UMacTargetSettings::StaticClass(), GetTransientPackage(), "MacTargetSettings", RF_Standalone);
+		TargetSettings = NewObject<UMacTargetSettings>(GetTransientPackage(), "MacTargetSettings", RF_Standalone);
 		TargetSettings->AddToRoot();
 
 		ISettingsModule* SettingsModule = FModuleManager::GetModulePtr<ISettingsModule>("Settings");
