Merging CL 2719488 - Enabled Metal shaders in cooked builds on Mac

diff --git a/Engine/Config/BaseEngine.ini b/Engine/Config/BaseEngine.ini
--- a/Engine/Config/BaseEngine.ini
+++ b/Engine/Config/BaseEngine.ini
@@ -2369,6 +2369,8 @@ MinimumOSVersion=MSOS_Vista
 +TargetedRHIs=GLSL_150
  
 [/Script/MacTargetPlatform.MacTargetSettings]
++TargetedRHIs=SF_METAL_SM4
++TargetedRHIs=SF_METAL_SM5
 +TargetedRHIs=GLSL_150_MAC
 
 [HMDPluginPriority]
diff --git a/Engine/Source/Developer/Mac/MacTargetPlatform/Private/GenericMacTargetPlatform.h b/Engine/Source/Developer/Mac/MacTargetPlatform/Private/GenericMacTargetPlatform.h
--- a/Engine/Source/Developer/Mac/MacTargetPlatform/Private/GenericMacTargetPlatform.h
+++ b/Engine/Source/Developer/Mac/MacTargetPlatform/Private/GenericMacTargetPlatform.h
@@ -116,10 +116,14 @@ return TSuper::SupportsFeature(Feature);
 		{
 			static FName NAME_GLSL_150_MAC(TEXT("GLSL_150_MAC"));
 			OutFormats.AddUnique(NAME_GLSL_150_MAC);
-			static FName NAME_SF_METAL_SM4(TEXT("SF_METAL_SM4"));
-			OutFormats.AddUnique(NAME_SF_METAL_SM4);
-			static FName NAME_SF_METAL_SM5(TEXT("SF_METAL_SM5"));
-			OutFormats.AddUnique(NAME_SF_METAL_SM5);
+
+			if (FPlatformMisc::MacOSXVersionCompare(10, 11, 0) >= 0)
+			{
+				static FName NAME_SF_METAL_SM4(TEXT("SF_METAL_SM4"));
+				OutFormats.AddUnique(NAME_SF_METAL_SM4);
+				static FName NAME_SF_METAL_SM5(TEXT("SF_METAL_SM5"));
+				OutFormats.AddUnique(NAME_SF_METAL_SM5);
+			}
 		}
 	}
 
