Enabled offline shader compilation for Metal (on Mac only), plus CIS fix

diff --git a/Engine/Source/Developer/Apple/MetalShaderFormat/Private/MetalShaderFormat.h b/Engine/Source/Developer/Apple/MetalShaderFormat/Private/MetalShaderFormat.h
--- a/Engine/Source/Developer/Apple/MetalShaderFormat/Private/MetalShaderFormat.h
+++ b/Engine/Source/Developer/Apple/MetalShaderFormat/Private/MetalShaderFormat.h
@@ -6,5 +6,5 @@ extern void CompileShader_Metal(const struct FShaderCompilerInput& Input, struct
 
 // Set this to 0 to get shader source in the graphics debugger
 // Note: Offline and runtime compiled shaders have separate DDC versions and can co-exist
-#define METAL_OFFLINE_COMPILE 0 // !PLATFORM_WINDOWS
+#define METAL_OFFLINE_COMPILE !PLATFORM_WINDOWS
 
diff --git a/Engine/Source/Developer/Mac/MacTargetPlatform/Private/GenericMacTargetPlatform.h b/Engine/Source/Developer/Mac/MacTargetPlatform/Private/GenericMacTargetPlatform.h
--- a/Engine/Source/Developer/Mac/MacTargetPlatform/Private/GenericMacTargetPlatform.h
+++ b/Engine/Source/Developer/Mac/MacTargetPlatform/Private/GenericMacTargetPlatform.h
@@ -117,6 +117,7 @@ return TSuper::SupportsFeature(Feature);
 			static FName NAME_GLSL_150_MAC(TEXT("GLSL_150_MAC"));
 			OutFormats.AddUnique(NAME_GLSL_150_MAC);
 
+#if PLATFORM_MAC // @todo: Enable on Windows, Linux and OS X 10.10 (fallback to online shader compiler)
 			if (FPlatformMisc::MacOSXVersionCompare(10, 11, 0) >= 0)
 			{
 				static FName NAME_SF_METAL_SM4(TEXT("SF_METAL_SM4"));
@@ -124,6 +125,7 @@ return TSuper::SupportsFeature(Feature);
 				static FName NAME_SF_METAL_SM5(TEXT("SF_METAL_SM5"));
 				OutFormats.AddUnique(NAME_SF_METAL_SM5);
 			}
+#endif
 		}
 	}
 
