Remove option for Metal 1.2 support from Mac target settings, Metal 2.0 is now the minimum.

diff --git a/Engine/Source/Developer/Mac/MacTargetPlatform/Classes/MacTargetSettings.h b/Engine/Source/Developer/Mac/MacTargetPlatform/Classes/MacTargetSettings.h
--- a/Engine/Source/Developer/Mac/MacTargetPlatform/Classes/MacTargetSettings.h
+++ b/Engine/Source/Developer/Mac/MacTargetPlatform/Classes/MacTargetSettings.h
@@ -13,9 +13,6 @@
 UENUM()
 enum class EMacMetalShaderStandard : uint8
 {
-    /** Metal Shaders, supporting Tessellation Shaders & Fragment Shader UAVs, Compatible With macOS Sierra 10.12.0 or later (std=osx-metal1.2) */
-    MacMetalSLStandard_1_2 = 2 UMETA(DisplayName="Metal v1.2 (10.12.0+)"),
-    
     /** Metal Shaders, supporting multiple viewports, Compatible With macOS 10.13.0 or later (std=osx-metal2.0) */
     MacMetalSLStandard_2_0 = 3 UMETA(DisplayName="Metal v2.0 (10.13.0+)"),
     
