Remove the Metal 2.1 + IABs shader version from the Mac target settings as this is an experimental feature that I don't want to expose to users as it has broken since I worked on it last.

diff --git a/Engine/Source/Developer/Mac/MacTargetPlatform/Classes/MacTargetSettings.h b/Engine/Source/Developer/Mac/MacTargetPlatform/Classes/MacTargetSettings.h
--- a/Engine/Source/Developer/Mac/MacTargetPlatform/Classes/MacTargetSettings.h
+++ b/Engine/Source/Developer/Mac/MacTargetPlatform/Classes/MacTargetSettings.h
@@ -18,9 +18,6 @@ enum class EMacMetalShaderStandard : uint8
     
     /** Metal Shaders, supporting indirect command buffers, Compatible With macOS 10.14.0 or later (std=osx-metal2.1) */
     MacMetalSLStandard_2_1 = 4 UMETA(DisplayName="Metal v2.1 (10.14.0+)"),
-	
-	/** Metal Shaders, using indirect argument buffers to reduce CPU cost, Compatible With macOS 10.14.0 or later (std=osx-metal2.1) */
-	MacMetalSLStandard_2_1_IABs = 5 UMETA(DisplayName="Metal v2.1 w/Indirect Argument Buffers (10.14.0+)"),
 };
 
 /**
