Flip on Stream Caching for Android, Mac and Windows on 11.10.

diff --git a/Engine/Source/Developer/Mac/MacTargetPlatform/Classes/MacTargetSettings.h b/Engine/Source/Developer/Mac/MacTargetPlatform/Classes/MacTargetSettings.h
--- a/Engine/Source/Developer/Mac/MacTargetPlatform/Classes/MacTargetSettings.h
+++ b/Engine/Source/Developer/Mac/MacTargetPlatform/Classes/MacTargetSettings.h
@@ -8,6 +8,7 @@
 
 #include "CoreMinimal.h"
 #include "UObject/Object.h"
+#include "AudioCompressionSettings.h"
 #include "MacTargetSettings.generated.h"
 
 UENUM()
@@ -106,4 +107,46 @@ public:
 	/** Which of the currently enabled occlusion plugins to use on Windows. */
 	UPROPERTY(config, EditAnywhere, Category = "Audio")
 	FString OcclusionPlugin;
+
+	/** Cook Overrides */
+
+	/** Various overrides for how this platform should handle compression and decompression */
+	UPROPERTY(GlobalConfig, EditAnywhere, Category = "Audio")
+	FPlatformRuntimeAudioCompressionOverrides CompressionOverrides;
+
+	/** When this is enabled, Actual compressed data will be separated from the USoundWave, and loaded into a cache. */
+	UPROPERTY(GlobalConfig, EditAnywhere, Category = "Audio|CookOverrides", meta = (DisplayName = "Use Stream Caching (Experimental)"))
+	bool bUseAudioStreamCaching;
+
+	/** This determines the max amount of memory that should be used for the cache at any given time. If set low (<= 8 MB), it lowers the size of individual chunks of audio during cook. */
+	UPROPERTY(GlobalConfig, EditAnywhere, Category = "Audio|CookOverrides|Stream Caching", meta = (DisplayName = "Max Cache Size (KB)"))
+	int32 CacheSizeKB;
+
+	UPROPERTY(GlobalConfig, EditAnywhere, Category = "Audio|CookOverrides")
+	bool bResampleForDevice;
+
+	/** Mapping of which sample rates are used for each sample rate quality for a specific platform. */
+
+	UPROPERTY(GlobalConfig, EditAnywhere, Category = "Audio|CookOverrides|ResamplingQuality", meta = (DisplayName = "Max"))
+	float MaxSampleRate;
+
+	UPROPERTY(GlobalConfig, EditAnywhere, Category = "Audio|CookOverrides|ResamplingQuality", meta = (DisplayName = "High"))
+	float HighSampleRate;
+
+	UPROPERTY(GlobalConfig, EditAnywhere, Category = "Audio|CookOverrides|ResamplingQuality", meta = (DisplayName = "Medium"))
+	float MedSampleRate;
+
+	UPROPERTY(GlobalConfig, EditAnywhere, Category = "Audio|CookOverrides|ResamplingQuality", meta = (DisplayName = "Low"))
+	float LowSampleRate;
+
+	UPROPERTY(GlobalConfig, EditAnywhere, Category = "Audio|CookOverrides|ResamplingQuality", meta = (DisplayName = "Min"))
+	float MinSampleRate;
+
+	/** Scales all compression qualities when cooking to this platform. For example, 0.5 will halve all compression qualities, and 1.0 will leave them unchanged. */
+	UPROPERTY(GlobalConfig, EditAnywhere, Category = "Audio|CookOverrides")
+	float CompressionQualityModifier;
+
+	/** When set to anything beyond 0, this will ensure any SoundWaves longer than this value, in seconds, to stream directly off of the disk. */
+	UPROPERTY(GlobalConfig, EditAnywhere, Category = "Audio|CookOverrides", meta = (DisplayName = "Stream All Soundwaves Longer Than: "))
+	float AutoStreamingThreshold;
 };
diff --git a/Engine/Source/Developer/Mac/MacTargetPlatform/Private/GenericMacTargetPlatform.h b/Engine/Source/Developer/Mac/MacTargetPlatform/Private/GenericMacTargetPlatform.h
--- a/Engine/Source/Developer/Mac/MacTargetPlatform/Private/GenericMacTargetPlatform.h
+++ b/Engine/Source/Developer/Mac/MacTargetPlatform/Private/GenericMacTargetPlatform.h
@@ -22,6 +22,122 @@
 
 #define LOCTEXT_NAMESPACE "TGenericMacTargetPlatform"
 
+namespace Mac
+{
+#if WITH_ENGINE
+	FORCEINLINE void CachePlatformAudioCookOverrides(FPlatformAudioCookOverrides& OutOverrides)
+	{
+		const TCHAR* CategoryName = TEXT("/Script/MacTargetPlatform.MacTargetSettings");
+
+		GConfig->GetBool(CategoryName, TEXT("bUseAudioStreamCaching"), OutOverrides.bUseStreamCaching, GEngineIni);
+
+		/** Memory Load On Demand Settings */
+		if (OutOverrides.bUseStreamCaching)
+		{
+			// Cache size:
+			int32 RetrievedCacheSize = 32 * 1024;
+			GConfig->GetInt(CategoryName, TEXT("CacheSizeKB"), RetrievedCacheSize, GEngineIni);
+			OutOverrides.StreamCachingSettings.CacheSizeKB = RetrievedCacheSize;
+		}
+
+		GConfig->GetBool(CategoryName, TEXT("bResampleForDevice"), OutOverrides.bResampleForDevice, GEngineIni);
+
+		GConfig->GetFloat(CategoryName, TEXT("CompressionQualityModifier"), OutOverrides.CompressionQualityModifier, GEngineIni);
+
+		GConfig->GetFloat(CategoryName, TEXT("AutoStreamingThreshold"), OutOverrides.AutoStreamingThreshold, GEngineIni);
+
+		//Cache sample rate map:
+		float RetrievedSampleRate = -1.0f;
+
+		GConfig->GetFloat(CategoryName, TEXT("MaxSampleRate"), RetrievedSampleRate, GEngineIni);
+		float* FoundSampleRate = OutOverrides.PlatformSampleRates.Find(ESoundwaveSampleRateSettings::Max);
+
+		if (FoundSampleRate)
+		{
+			if (!FMath::IsNearlyEqual(*FoundSampleRate, RetrievedSampleRate))
+			{
+				*FoundSampleRate = RetrievedSampleRate;
+			}
+
+		}
+		else
+		{
+			OutOverrides.PlatformSampleRates.Add(ESoundwaveSampleRateSettings::Max, RetrievedSampleRate);
+		}
+
+		RetrievedSampleRate = -1.0f;
+
+		GConfig->GetFloat(CategoryName, TEXT("HighSampleRate"), RetrievedSampleRate, GEngineIni);
+		FoundSampleRate = OutOverrides.PlatformSampleRates.Find(ESoundwaveSampleRateSettings::High);
+
+		if (FoundSampleRate)
+		{
+			if (!FMath::IsNearlyEqual(*FoundSampleRate, RetrievedSampleRate))
+			{
+				*FoundSampleRate = RetrievedSampleRate;
+			}
+
+		}
+		else
+		{
+			OutOverrides.PlatformSampleRates.Add(ESoundwaveSampleRateSettings::High, RetrievedSampleRate);
+		}
+
+
+		RetrievedSampleRate = -1.0f;
+
+		GConfig->GetFloat(CategoryName, TEXT("MedSampleRate"), RetrievedSampleRate, GEngineIni);
+		FoundSampleRate = OutOverrides.PlatformSampleRates.Find(ESoundwaveSampleRateSettings::Medium);
+
+		if (FoundSampleRate)
+		{
+			if (!FMath::IsNearlyEqual(*FoundSampleRate, RetrievedSampleRate))
+			{
+				*FoundSampleRate = RetrievedSampleRate;
+			}
+		}
+		else
+		{
+			OutOverrides.PlatformSampleRates.Add(ESoundwaveSampleRateSettings::Medium, RetrievedSampleRate);
+		}
+
+		RetrievedSampleRate = -1.0f;
+
+		GConfig->GetFloat(CategoryName, TEXT("LowSampleRate"), RetrievedSampleRate, GEngineIni);
+		FoundSampleRate = OutOverrides.PlatformSampleRates.Find(ESoundwaveSampleRateSettings::Low);
+
+		if (FoundSampleRate)
+		{
+			if (!FMath::IsNearlyEqual(*FoundSampleRate, RetrievedSampleRate))
+			{
+				*FoundSampleRate = RetrievedSampleRate;
+			}
+		}
+		else
+		{
+			OutOverrides.PlatformSampleRates.Add(ESoundwaveSampleRateSettings::Low, RetrievedSampleRate);
+		}
+
+		RetrievedSampleRate = -1.0f;
+
+		GConfig->GetFloat(CategoryName, TEXT("MinSampleRate"), RetrievedSampleRate, GEngineIni);
+		FoundSampleRate = OutOverrides.PlatformSampleRates.Find(ESoundwaveSampleRateSettings::Min);
+
+		if (FoundSampleRate)
+		{
+			if (!FMath::IsNearlyEqual(*FoundSampleRate, RetrievedSampleRate))
+			{
+				*FoundSampleRate = RetrievedSampleRate;
+			}
+		}
+		else
+		{
+			OutOverrides.PlatformSampleRates.Add(ESoundwaveSampleRateSettings::Min, RetrievedSampleRate);
+		}
+	}
+#endif
+}
+
 /**
  * Template for Mac target platforms
  */
@@ -343,7 +459,17 @@ return TSuper::SupportsFeature(Feature);
 
 	virtual FPlatformAudioCookOverrides* GetAudioCompressionSettings() const override
 	{
-		return nullptr;
+		static FPlatformAudioCookOverrides Settings;
+
+		static bool bCachedPlatformSettings = false;
+
+		if (!bCachedPlatformSettings)
+		{
+			Mac::CachePlatformAudioCookOverrides(Settings);
+			bCachedPlatformSettings = true;
+		}
+
+		return &Settings;
 	}
 
 #endif //WITH_ENGINE
diff --git a/Engine/Source/Runtime/Engine/Private/AudioCompressionSettingsUtils.cpp b/Engine/Source/Runtime/Engine/Private/AudioCompressionSettingsUtils.cpp
--- a/Engine/Source/Runtime/Engine/Private/AudioCompressionSettingsUtils.cpp
+++ b/Engine/Source/Runtime/Engine/Private/AudioCompressionSettingsUtils.cpp
@@ -76,6 +76,8 @@ void CacheCurrentPlatformAudioCookOverrides(FPlatformAudioCookOverrides& OutOver
 {
 #if PLATFORM_ANDROID && !PLATFORM_LUMIN
 	const TCHAR* CategoryName = TEXT("/Script/AndroidRuntimeSettings.AndroidRuntimeSettings");
+#elif PLATFORM_MAC
+	const TCHAR* CategoryName = TEXT("/Script/MacTargetPlatform.MacTargetSettings");
 #elif PLATFORM_IOS
 	const TCHAR* CategoryName = TEXT("/Script/IOSRuntimeSettings.IOSRuntimeSettings");
 #elif PLATFORM_SWITCH
