Fix the description used for Mac platform target architecture in the editor project settings.

#rnx

diff --git a/Engine/Source/Developer/Mac/MacTargetPlatform/Classes/MacTargetSettings.h b/Engine/Source/Developer/Mac/MacTargetPlatform/Classes/MacTargetSettings.h
--- a/Engine/Source/Developer/Mac/MacTargetPlatform/Classes/MacTargetSettings.h
+++ b/Engine/Source/Developer/Mac/MacTargetPlatform/Classes/MacTargetSettings.h
@@ -59,8 +59,8 @@ public:
 	TArray<FString> TargetedRHIs;
     
     /**
-     * The maximum supported Metal shader langauge version.
-     * This defines what features may be used and OS versions supported.
+     * The target Mac platform CPU architecture.
+     * This defines which CPU architectures to target: x86_64 (Intel), arm64 (Apple Silicon) or Universal (Intel & Apple Silicon).
      */
     UPROPERTY(EditAnywhere, config, Category=Packaging, meta = (DisplayName = "Architectures to Package For (Experimental)"))
     EMacTargetArchitecture TargetArchitecture;
