Fix for GenericMaxTargetPlatform.h and texture building

diff --git a/Engine/Source/Developer/Mac/MacTargetPlatform/Private/GenericMacTargetPlatform.h b/Engine/Source/Developer/Mac/MacTargetPlatform/Private/GenericMacTargetPlatform.h
--- a/Engine/Source/Developer/Mac/MacTargetPlatform/Private/GenericMacTargetPlatform.h
+++ b/Engine/Source/Developer/Mac/MacTargetPlatform/Private/GenericMacTargetPlatform.h
@@ -157,7 +157,7 @@ public:
 		if (!IS_DEDICATED_SERVER)
 		{
 			// just use the standard texture format name for this texture (with DX11 support)
-			GetDefaultTextureFormatNamePerLayer(OutFormats.AddDefaulted_GetRef(), this, Texture, true);
+			GetDefaultTextureFormatNamePerLayer(OutFormats.AddDefaulted_GetRef(), this, Texture);
 		}
 	}
 
