Fixed GetAllDefaultTextureFormats as well for MacEditor

diff --git a/Engine/Source/Developer/Mac/MacTargetPlatform/Private/GenericMacTargetPlatform.h b/Engine/Source/Developer/Mac/MacTargetPlatform/Private/GenericMacTargetPlatform.h
--- a/Engine/Source/Developer/Mac/MacTargetPlatform/Private/GenericMacTargetPlatform.h
+++ b/Engine/Source/Developer/Mac/MacTargetPlatform/Private/GenericMacTargetPlatform.h
@@ -166,7 +166,7 @@ public:
 		if (!IS_DEDICATED_SERVER)
 		{
 			// just use the standard texture format name for this texture (with DX11 support)
-			GetAllDefaultTextureFormats(this, OutFormats, true);
+			GetAllDefaultTextureFormats(this, OutFormats);
 		}
 	}
 
