- Added WITH_ENGINE wrappers to some code in MacTargetPlatform that doesn't work without engine

diff --git a/Engine/Source/Developer/Mac/MacTargetPlatform/Private/MacTargetPlatformModule.cpp b/Engine/Source/Developer/Mac/MacTargetPlatform/Private/MacTargetPlatformModule.cpp
--- a/Engine/Source/Developer/Mac/MacTargetPlatform/Private/MacTargetPlatformModule.cpp
+++ b/Engine/Source/Developer/Mac/MacTargetPlatform/Private/MacTargetPlatformModule.cpp
@@ -41,6 +41,7 @@ public:
 
 	virtual void StartupModule() override
 	{
+#if WITH_ENGINE
 		TargetSettings = NewObject<UMacTargetSettings>(GetTransientPackage(), "MacTargetSettings", RF_Standalone);
 		
 		// We need to manually load the config properties here, as this module is loaded before the UObject system is setup to do this
@@ -78,10 +79,12 @@ public:
 				TargetSettings
 			);
 		}
+#endif
 	}
 
 	virtual void ShutdownModule() override
 	{
+#if WITH_ENGINE
 		ISettingsModule* SettingsModule = FModuleManager::GetModulePtr<ISettingsModule>("Settings");
 
 		if (SettingsModule != nullptr)
@@ -98,6 +101,7 @@ public:
 		{
 			TargetSettings = NULL;
 		}
+#endif
 	}
 
 	// End IModuleInterface interface
