- Updated Mac compile default to Host architecture
- Added missing Project Settings
- Added support for BuildMachine haaving unique defaults

diff --git a/Engine/Config/BaseEngine.ini b/Engine/Config/BaseEngine.ini
--- a/Engine/Config/BaseEngine.ini
+++ b/Engine/Config/BaseEngine.ini
@@ -2894,8 +2894,9 @@ UseFastIntrinsics=False
 EnableMathOptimisations=True
 TargetArchitecture=MacTargetArchitectureUniversal
 EditorTargetArchitecture=MacTargetArchitectureUniversal
-EditorDefaultArchitecture=MacTargetArchitectureIntel
-DefaultArchitecture=MacTargetArchitectureIntel
+EditorDefaultArchitecture=MacTargetArchitectureHost
+DefaultArchitecture=MacTargetArchitectureHost
+bBuildAllSupportedOnBuildMachine=true
 bEnableRayTracing=false
 
 [HMDPluginPriority]
diff --git a/Engine/Source/Developer/Mac/MacTargetPlatform/Classes/MacTargetSettings.h b/Engine/Source/Developer/Mac/MacTargetPlatform/Classes/MacTargetSettings.h
--- a/Engine/Source/Developer/Mac/MacTargetPlatform/Classes/MacTargetSettings.h
+++ b/Engine/Source/Developer/Mac/MacTargetPlatform/Classes/MacTargetSettings.h
@@ -35,8 +35,11 @@ enum class EMacTargetArchitecture : uint8
     /** Create Universal packages that run natively on all Macs */
     MacTargetArchitectureUniversal = 1 UMETA(DisplayName="Universal (Intel & Apple Silicon)"),
 
-    /** Create packages that can run natively on Apple Silicon Macs */
-    MacTargetArchitectureAppleSil = 2 UMETA(DisplayName="Apple Silicon"),
+	/** Create packages that can run natively on Apple Silicon Macs, and will not run on Intel Macs */
+	MacTargetArchitectureAppleSilicon = 2 UMETA(DisplayName="Apple Silicon"),
+
+	/** Create  packages that match the architecture of the Mac the package is made on */
+	MacTargetArchitectureHost = 3 UMETA(DisplayName="Host (Matches current machine)"),
 };
 
 
@@ -57,14 +60,45 @@ public:
 	 */
 	UPROPERTY(EditAnywhere, config, Category=Rendering)
 	TArray<FString> TargetedRHIs;
-    
-    /**
-     * The target Mac platform CPU architecture.
-     * This defines which CPU architectures to target: x86_64 (Intel), arm64 (Apple Silicon) or Universal (Intel & Apple Silicon).
-     */
-    UPROPERTY(EditAnywhere, config, Category=Packaging, meta = (DisplayName = "Architectures to Package For (Experimental)"))
-    EMacTargetArchitecture TargetArchitecture;
-    
+	
+	/**
+	 * The set of architecture(s) this project supports for Editor builds
+	 * This defines which CPU architectures to target: x86_64 (Intel), arm64 (Apple Silicon) or Universal (Intel & Apple Silicon).
+	 * It is recommended to use Universal unless you have editor plugins or other libraries that do not support Apple Silicon
+	 */
+	UPROPERTY(EditAnywhere, config, Category=Packaging, meta = (DisplayName = "Supported Architecture(s) for Editor builds"))
+	EMacTargetArchitecture EditorTargetArchitecture;
+	
+	/**
+	 * The target Mac platform CPU architecture.
+	 * This defines which CPU architectures to target: x86_64 (Intel), arm64 (Apple Silicon) or Universal (Intel & Apple Silicon).
+	 * It is recommended to use Universal unless you have runtime plugins or other libraries that do not support Apple Silicon
+	 */
+	UPROPERTY(EditAnywhere, config, Category=Packaging, meta = (DisplayName = "Supported Architecture(s) for non-Editor builds"))
+	EMacTargetArchitecture TargetArchitecture;
+	
+	/**
+	 * The architecture to compile the Editor target
+	 * This defines which CPU architectures to target: x86_64 (Intel), arm64 (Apple Silicon) or Universal (Intel & Apple Silicon), or Host to match the machine doing the building
+	 * Can override with -Architecture= on the UBT commandline, or -EditorArchitecture=  on the BuildCookRun commandline
+	 */
+	UPROPERTY(EditAnywhere, config, Category=Packaging, meta = (DisplayName = "Architecture(s) when building Editor"))
+	EMacTargetArchitecture EditorDefaultArchitecture;
+	
+	/**
+	 * The architectures to compile non-Editor (games, programs, etc) targets for builds outside of Xcode
+	 * This defines which CPU architectures to target: x86_64 (Intel), arm64 (Apple Silicon) or Universal (Intel & Apple Silicon), or Host to match the machine doing the building
+	 * Can override with -Architecture= on the UBT commandline, or -GameArchitecture= or -ProgramArchitecture= on the BuildCookRun commandline
+	 */
+	UPROPERTY(EditAnywhere, config, Category=Packaging, meta = (DisplayName = "Architecture(s) when building non-Editor"))
+	EMacTargetArchitecture DefaultArchitecture;
+	
+	/**
+	 * If true, builds running on BuildMachines (when the 'IsBuildMachine' environment variable is set to 1) will compile all Supported architectures
+	 */
+	UPROPERTY(EditAnywhere, config, Category=Packaging, meta = (DisplayName = "Build all supported Architectures on Build Machines"))
+	bool bBuildAllSupportedOnBuildMachine;
+
     /**
      * The Metal shader language version which will be used when compiling the shaders.
      */
diff --git a/Engine/Source/Programs/UnrealBuildTool/Platform/Mac/UEBuildMac.cs b/Engine/Source/Programs/UnrealBuildTool/Platform/Mac/UEBuildMac.cs
--- a/Engine/Source/Programs/UnrealBuildTool/Platform/Mac/UEBuildMac.cs
+++ b/Engine/Source/Programs/UnrealBuildTool/Platform/Mac/UEBuildMac.cs
@@ -153,6 +153,7 @@ namespace UnrealBuildTool
 		private UnrealArchitectures GetProjectArchitectures(FileReference? ProjectFile, string? TargetName, bool bGetAllSupported, bool bIsDistributionMode)
 		{
 			bool bIsEditor = false;
+			bool bIsBuildMachine = Environment.GetEnvironmentVariable("IsBuildMachine") == "1";
 
 			// get project ini from ProjetFile, or if null, then try to get it from the target rules
 			if (TargetName != null)
@@ -188,8 +189,10 @@ namespace UnrealBuildTool
 			string DefaultKey = bIsEditor ? "EditorDefaultArchitecture" : "DefaultArchitecture";
 			string SupportedArchitecture;
 			string DefaultArchitecture;
+			bool bBuildAllSupportedOnBuildMachine;
 			EngineIni.GetString("/Script/MacTargetPlatform.MacTargetSettings", SupportKey, out SupportedArchitecture);
 			EngineIni.GetString("/Script/MacTargetPlatform.MacTargetSettings", DefaultKey, out DefaultArchitecture);
+			EngineIni.GetBool("/Script/MacTargetPlatform.MacTargetSettings", "bBuildAllSupportedOnBuildMachine", out bBuildAllSupportedOnBuildMachine);
 			SupportedArchitecture = SupportedArchitecture.ToLower();
 			DefaultArchitecture = DefaultArchitecture.ToLower();
 
@@ -199,14 +202,15 @@ namespace UnrealBuildTool
 			// make sure we found a good value
 			if (!bSupportsArm64 && !bSupportsX86)
 			{
-				throw new BuildException($"Unknown {DefaultKey} value found ('{DefaultArchitecture}') in .ini");
+				throw new BuildException($"Unknown {SupportKey} value found ('{SupportedArchitecture}') in .ini");
 			}
 
 			// choose a supported architecture(s) based on desired type
 			List<UnrealArch> Architectures = new();
 
 			// return all supported if getting supported, compiling for distribution, or we want active, and "all" is selected
-			if (bGetAllSupported || bIsDistributionMode || DefaultArchitecture.Equals("all", StringComparison.InvariantCultureIgnoreCase))
+			if (bGetAllSupported || bIsDistributionMode || DefaultArchitecture.Equals("all", StringComparison.InvariantCultureIgnoreCase) ||
+				(bIsBuildMachine && bBuildAllSupportedOnBuildMachine))
 			{
 				if (bSupportsArm64)
 				{
@@ -217,7 +221,7 @@ namespace UnrealBuildTool
 					Architectures.Add(UnrealArch.X64);
 				}
 			}
-			else if (DefaultArchitecture.Equals("host", StringComparison.InvariantCultureIgnoreCase))
+			else if (DefaultArchitecture.Contains("host"))
 			{
 				if (MacExports.IsRunningOnAppleArchitecture && bSupportsArm64)
 				{
