Fix LOCTEXT conflicts, also fix Mac settings being replaced

diff --git a/Engine/Source/Developer/Mac/MacPlatformEditor/Private/MacPlatformEditorModule.cpp b/Engine/Source/Developer/Mac/MacPlatformEditor/Private/MacPlatformEditorModule.cpp
--- a/Engine/Source/Developer/Mac/MacPlatformEditor/Private/MacPlatformEditorModule.cpp
+++ b/Engine/Source/Developer/Mac/MacPlatformEditor/Private/MacPlatformEditorModule.cpp
@@ -28,13 +28,13 @@ class FMacPlatformEditorModule
 		if (SettingsModule != nullptr)
 		{
 			SettingsModule->RegisterSettings("Project", "Platforms", "Mac",
-				LOCTEXT("TargetSettingsName", "Mac"),
-				LOCTEXT("TargetSettingsDescription", "Settings for Mac target platform"),
+				LOCTEXT("MacTargetSettingsName", "Mac"),
+				LOCTEXT("MacTargetSettingsDescription", "Settings for Mac target platform"),
 				GetMutableDefault<UMacTargetSettings>()
 			);
-            SettingsModule->RegisterSettings("Project", "Platforms", "Mac",
-                LOCTEXT("TargetSettingsName", "Xcode projects"),
-                LOCTEXT("TargetSettingsDescription", "Settings for Xcode projects"),
+            SettingsModule->RegisterSettings("Project", "Platforms", "Xcode",
+                LOCTEXT("XcodeProjectSettingsName", "Xcode projects"),
+                LOCTEXT("XcodeProjectSettingsDescription", "Settings for Xcode projects"),
                 GetMutableDefault<UXcodeProjectSettings>()
             );
 		}
diff --git a/Engine/Source/Developer/Mac/MacTargetPlatform/Private/MacTargetPlatformModule.cpp b/Engine/Source/Developer/Mac/MacTargetPlatform/Private/MacTargetPlatformModule.cpp
--- a/Engine/Source/Developer/Mac/MacTargetPlatform/Private/MacTargetPlatformModule.cpp
+++ b/Engine/Source/Developer/Mac/MacTargetPlatform/Private/MacTargetPlatformModule.cpp
@@ -78,13 +78,13 @@ public:
 		if (SettingsModule != nullptr)
 		{
 			SettingsModule->RegisterSettings("Project", "Platforms", "Mac",
-				LOCTEXT("TargetSettingsName", "Mac"),
-				LOCTEXT("TargetSettingsDescription", "Settings and resources for Mac platform"),
+				LOCTEXT("MacTargetSettingsName", "Mac"),
+				LOCTEXT("MacTargetSettingsDescription", "Settings and resources for Mac platform"),
 				TargetSettings
 			);
-            SettingsModule->RegisterSettings("Project", "Platforms", "Mac",
-                LOCTEXT("TargetSettingsName", "Xcode projects"),
-                LOCTEXT("TargetSettingsDescription", "Settings for Xcode projects"),
+            SettingsModule->RegisterSettings("Project", "Platforms", "Xcode",
+                LOCTEXT("XcodeProjectSettingsName", "Xcode projects"),
+                LOCTEXT("XcodeProjectSettingsDescription", "Settings for Xcode projects"),
                 ProjectSettings
             );
 		}
