dynamic settings enabling for Xcode project settings

diff --git a/Engine/Source/Developer/Mac/MacPlatformEditor/Private/MacPlatformEditorModule.cpp b/Engine/Source/Developer/Mac/MacPlatformEditor/Private/MacPlatformEditorModule.cpp
--- a/Engine/Source/Developer/Mac/MacPlatformEditor/Private/MacPlatformEditorModule.cpp
+++ b/Engine/Source/Developer/Mac/MacPlatformEditor/Private/MacPlatformEditorModule.cpp
@@ -33,7 +33,7 @@ class FMacPlatformEditorModule
 				GetMutableDefault<UMacTargetSettings>()
 			);
             SettingsModule->RegisterSettings("Project", "Platforms", "Xcode",
-                LOCTEXT("XcodeProjectSettingsName", "Xcode projects"),
+                LOCTEXT("XcodeProjectSettingsName", "Xcode Projects"),
                 LOCTEXT("XcodeProjectSettingsDescription", "Settings for Xcode projects"),
                 GetMutableDefault<UXcodeProjectSettings>()
             );
@@ -47,6 +47,7 @@ class FMacPlatformEditorModule
 		if (SettingsModule != nullptr)
 		{
 			SettingsModule->UnregisterSettings("Project", "Platforms", "Mac");
+            SettingsModule->UnregisterSettings("Project", "Platforms", "Xcode");
 		}
 	}
 };
diff --git a/Engine/Source/Developer/Mac/MacTargetPlatform/Classes/XcodeProjectSettings.h b/Engine/Source/Developer/Mac/MacTargetPlatform/Classes/XcodeProjectSettings.h
--- a/Engine/Source/Developer/Mac/MacTargetPlatform/Classes/XcodeProjectSettings.h
+++ b/Engine/Source/Developer/Mac/MacTargetPlatform/Classes/XcodeProjectSettings.h
@@ -30,60 +30,60 @@ public:
     /**
      * Enable native Xcode code signing
      */
-    UPROPERTY(EditAnywhere, config, Category=Xcode, meta = (DisplayName = "Use Modern Code Signing"))
+    UPROPERTY(EditAnywhere, config, Category=Xcode, meta = (EditCondition="bUseModernXcode", DisplayName = "Use Modern Code Signing"))
     bool bUseModernCodeSigning;
     
     /**
      * Team ID used for native Xcode code signing
      */
-    UPROPERTY(EditAnywhere, config, Category=Xcode, meta = (DisplayName = "Modern Code Sign Team"))
+    UPROPERTY(EditAnywhere, config, Category=Xcode, meta = (EditCondition="bUseModernXcode && bUseModernCodeSigning", DisplayName = "Modern Code Sign Team"))
     FString ModernSigningTeam;
     
     /**
      * Bundle ID prefix used for native Xcode code signing
      */
-    UPROPERTY(EditAnywhere, config, Category=Xcode, meta = (DisplayName = "ModernSigningPrefix"))
+    UPROPERTY(EditAnywhere, config, Category=Xcode, meta = (EditCondition="bUseModernXcode && bUseModernCodeSigning", DisplayName = "Bundle ID Prefix"))
     FString ModernSigningPrefix;
     
     /**
      * Bundle ID used for nativr Xcode code signing
      */
-    UPROPERTY(EditAnywhere, config, Category=Xcode, meta = (DisplayName = "ModernBundleIdentifier"))
+    UPROPERTY(EditAnywhere, config, Category=Xcode, meta = (EditCondition="bUseModernXcode && bUseModernCodeSigning", DisplayName = "Bundle ID"))
     FString ModernBundleIdentifier;
     
     /**
      * The App Category that will be used for Mac App Store submission
      */
-    UPROPERTY(EditAnywhere, config, Category=Xcode, meta = (DisplayName = "AppCategory"))
+    UPROPERTY(EditAnywhere, config, Category=Xcode, meta = (EditCondition="bUseModernXcode", DisplayName = "App Category"))
     FString AppCategory;
     
     /**
      * The template info.plist used for Mac game targets
      */
-    UPROPERTY(EditAnywhere, config, Category=Xcode, meta = (DisplayName = "TemplateMacPlist"))
+    UPROPERTY(EditAnywhere, config, Category=Xcode, meta = (EditCondition="bUseModernXcode", DisplayName = "Template Mac Target Info.plist"))
     FFilePath TemplateMacPlist;
     
     /**
      * The template info.plist used for Mac editor targets
      */
-    UPROPERTY(EditAnywhere, config, Category=Xcode, meta = (DisplayName = "TemplateMacEditorPlist"))
+    UPROPERTY(EditAnywhere, config, Category=Xcode, meta = (EditCondition="bUseModernXcode", DisplayName = "Template Mac Editor Info.plist"))
     FFilePath TemplateMacEditorPlist;
     
     /**
      * The template info.plist used for iOS game targets
      */
-    UPROPERTY(EditAnywhere, config, Category=Xcode, meta = (DisplayName = "TemplateIOSPlist"))
+    UPROPERTY(EditAnywhere, config, Category=Xcode, meta = (EditCondition="bUseModernXcode", DisplayName = "Template iOS Target Info.plist"))
     FFilePath TemplateIOSPlist;
     
     /**
      * The premade entitlement file used for development Mac builds
      */
-    UPROPERTY(EditAnywhere, config, Category=Xcode, meta = (DisplayName = "PremadeMacEntitlements"))
+    UPROPERTY(EditAnywhere, config, Category=Xcode, meta = (EditCondition="bUseModernXcode", DisplayName = "Development Mac Entitlement"))
     FFilePath PremadeMacEntitlements;
     
     /**
      * The premade entitlement file used for shipping Mac builds
      */
-    UPROPERTY(EditAnywhere, config, Category=Xcode, meta = (DisplayName = "ShippingSpecificMacEntitlements"))
+    UPROPERTY(EditAnywhere, config, Category=Xcode, meta = (EditCondition="bUseModernXcode", DisplayName = "Shipping Mac Entitlement"))
     FFilePath ShippingSpecificMacEntitlements;
 };
diff --git a/Engine/Source/Developer/Mac/MacTargetPlatform/Private/MacTargetPlatformModule.cpp b/Engine/Source/Developer/Mac/MacTargetPlatform/Private/MacTargetPlatformModule.cpp
--- a/Engine/Source/Developer/Mac/MacTargetPlatform/Private/MacTargetPlatformModule.cpp
+++ b/Engine/Source/Developer/Mac/MacTargetPlatform/Private/MacTargetPlatformModule.cpp
@@ -83,7 +83,7 @@ public:
 				TargetSettings
 			);
             SettingsModule->RegisterSettings("Project", "Platforms", "Xcode",
-                LOCTEXT("XcodeProjectSettingsName", "Xcode projects"),
+                LOCTEXT("XcodeProjectSettingsName", "Xcode Projects"),
                 LOCTEXT("XcodeProjectSettingsDescription", "Settings for Xcode projects"),
                 ProjectSettings
             );
@@ -99,6 +99,7 @@ public:
 		if (SettingsModule != nullptr)
 		{
 			SettingsModule->UnregisterSettings("Project", "Platforms", "Mac");
+            SettingsModule->UnregisterSettings("Project", "Platforms", "Xcode");
 		}
 
 		if (!GExitPurge)
