add import buttons for Xcode project settings plist/entitlements

diff --git a/Engine/Config/BaseEngine.ini b/Engine/Config/BaseEngine.ini
--- a/Engine/Config/BaseEngine.ini
+++ b/Engine/Config/BaseEngine.ini
@@ -3037,8 +3037,8 @@ ModernBundleIdentifier=$(UE_SIGNING_PREFIX).$(UE_PRODUCT_NAME_STRIPPED)
 AppCategory=public.app-category.games
 TemplateMacPlist=(FilePath="/Game/Build/Mac/Resources/Info.Template.plist")
 TemplateIOSPlist=(FilePath="/Game/Build/IOS/UBTGenerated/Info.Template.plist")
-PremadeMacEntitlements=(FilePath="/Engine/Build/Mac/Resources/Sandbox.Server.entitlements")
-ShippingSpecificMacEntitlements=(FilePath="/Engine/Build/Mac/Resources/Sandbox.NoNet.entitlements")
+PremadeMacEntitlements=(FilePath="/Game/Build/Mac/Resources/Sandbox.Server.entitlements")
+ShippingSpecificMacEntitlements=(FilePath="/Game/Build/Mac/Resources/Sandbox.NoNet.entitlements")
 
 [HMDPluginPriority]
 ; Since SteamVR also works with the Oculus Rift, give priority to the native Oculus plugin before trying SteamVR.
diff --git a/Engine/Source/Developer/Mac/MacTargetPlatform/Classes/XcodeProjectSettings.h b/Engine/Source/Developer/Mac/MacTargetPlatform/Classes/XcodeProjectSettings.h
--- a/Engine/Source/Developer/Mac/MacTargetPlatform/Classes/XcodeProjectSettings.h
+++ b/Engine/Source/Developer/Mac/MacTargetPlatform/Classes/XcodeProjectSettings.h
@@ -8,8 +8,12 @@
 
 #include "CoreMinimal.h"
 #include "UObject/Object.h"
+#include "IDetailCustomization.h"
 #include "XcodeProjectSettings.generated.h"
 
+class IPropertyHandle;
+class FReply;
+
 /**
  * Implements the settings for Xcode projects
  */
@@ -60,30 +64,43 @@ public:
     /**
      * The template info.plist used for Mac game targets
      */
-    UPROPERTY(EditAnywhere, config, Category=Xcode, meta = (EditCondition="bUseModernXcode", DisplayName = "Template Mac Target Info.plist"))
+    UPROPERTY(EditAnywhere, config, Category="Info.plist", meta = (EditCondition="bUseModernXcode", DisplayName = "Mac Target Info.plist"))
     FFilePath TemplateMacPlist;
     
-    /**
-     * The template info.plist used for Mac editor targets
-     */
-    UPROPERTY(EditAnywhere, config, Category=Xcode, meta = (EditCondition="bUseModernXcode", DisplayName = "Template Mac Editor Info.plist"))
-    FFilePath TemplateMacEditorPlist;
-    
     /**
      * The template info.plist used for iOS game targets
      */
-    UPROPERTY(EditAnywhere, config, Category=Xcode, meta = (EditCondition="bUseModernXcode", DisplayName = "Template iOS Target Info.plist"))
+    UPROPERTY(EditAnywhere, config, Category="Info.plist", meta = (EditCondition="bUseModernXcode", DisplayName = "iOS Target Info.plist"))
     FFilePath TemplateIOSPlist;
     
     /**
      * The premade entitlement file used for development Mac builds
      */
-    UPROPERTY(EditAnywhere, config, Category=Xcode, meta = (EditCondition="bUseModernXcode", DisplayName = "Development Mac Entitlement"))
+    UPROPERTY(EditAnywhere, config, Category="Entitlement", meta = (EditCondition="bUseModernXcode", DisplayName = "Development Mac Entitlement"))
     FFilePath PremadeMacEntitlements;
     
     /**
      * The premade entitlement file used for shipping Mac builds
      */
-    UPROPERTY(EditAnywhere, config, Category=Xcode, meta = (EditCondition="bUseModernXcode", DisplayName = "Shipping Mac Entitlement"))
+    UPROPERTY(EditAnywhere, config, Category="Entitlement", meta = (EditCondition="bUseModernXcode", DisplayName = "Shipping Mac Entitlement"))
     FFilePath ShippingSpecificMacEntitlements;
 };
+
+class FXcodeProjectSettingsDetailsCustomization : public IDetailCustomization
+{
+public:
+    /** Makes a new instance of this detail layout class for a specific detail view requesting it */
+    static TSharedRef<IDetailCustomization> MakeInstance();
+
+    // IDetailCustomization interface
+    virtual void CustomizeDetails(IDetailLayoutBuilder& DetailLayout) override;
+    // End of IDetailCustomization interface
+    
+    FReply OnRestorePlistClicked();
+    FReply OnRestoreEntitlementClicked();
+    
+    TSharedPtr<IPropertyHandle> TemplateMacPlist;
+    TSharedPtr<IPropertyHandle> TemplateIOSPlist;
+    TSharedPtr<IPropertyHandle> PremadeMacEntitlements;
+    TSharedPtr<IPropertyHandle> ShippingEntitlements;
+};
diff --git a/Engine/Source/Developer/Mac/MacTargetPlatform/MacTargetPlatform.Build.cs b/Engine/Source/Developer/Mac/MacTargetPlatform/MacTargetPlatform.Build.cs
--- a/Engine/Source/Developer/Mac/MacTargetPlatform/MacTargetPlatform.Build.cs
+++ b/Engine/Source/Developer/Mac/MacTargetPlatform/MacTargetPlatform.Build.cs
@@ -12,6 +12,10 @@ public class MacTargetPlatform : ModuleRules
 				"CoreUObject",
 				"TargetPlatform",
 				"DesktopPlatform",
+                "InputCore",
+                "SlateCore",
+                "Slate",
+                "SourceControl"
 			}
 		);
 
diff --git a/Engine/Source/Developer/Mac/MacTargetPlatform/Private/MacTargetPlatformModule.cpp b/Engine/Source/Developer/Mac/MacTargetPlatform/Private/MacTargetPlatformModule.cpp
--- a/Engine/Source/Developer/Mac/MacTargetPlatform/Private/MacTargetPlatformModule.cpp
+++ b/Engine/Source/Developer/Mac/MacTargetPlatform/Private/MacTargetPlatformModule.cpp
@@ -3,6 +3,7 @@
 #include "CoreMinimal.h"
 #include "GenericMacTargetPlatform.h"
 #include "Modules/ModuleManager.h"
+#include "PropertyEditorModule.h"
 #include "ISettingsModule.h"
 #include "Interfaces/ITargetPlatformModule.h"
 #include "Modules/ModuleManager.h"
@@ -74,6 +75,12 @@ public:
         ProjectSettings->AddToRoot();
 
 		ISettingsModule* SettingsModule = FModuleManager::GetModulePtr<ISettingsModule>("Settings");
+        
+        FPropertyEditorModule& PropertyModule = FModuleManager::LoadModuleChecked<FPropertyEditorModule>("PropertyEditor");
+        PropertyModule.RegisterCustomClassLayout(
+            "XcodeProjectSettings",
+            FOnGetDetailCustomizationInstance::CreateStatic(&FXcodeProjectSettingsDetailsCustomization::MakeInstance)
+        );
 
 		if (SettingsModule != nullptr)
 		{
diff --git a/Engine/Source/Developer/Mac/MacTargetPlatform/Private/XcodeProjectSettings.cpp b/Engine/Source/Developer/Mac/MacTargetPlatform/Private/XcodeProjectSettings.cpp
new file mode 100644
--- /dev/null
+++ b/Engine/Source/Developer/Mac/MacTargetPlatform/Private/XcodeProjectSettings.cpp
@@ -0,0 +1,153 @@
+// Copyright Epic Games, Inc. All Rights Reserved.
+
+#include "XcodeProjectSettings.h"
+#include "HAL/FileManager.h"
+#include "DetailCategoryBuilder.h"
+#include "DetailLayoutBuilder.h"
+#include "DetailWidgetRow.h"
+#include "Widgets/Input/SButton.h"
+#include "SourceControlHelpers.h"
+#include "Framework/Notifications/NotificationManager.h"
+#include "Widgets/Notifications/SNotificationList.h"
+#include "PropertyHandle.h"
+
+#define LOCTEXT_NAMESPACE "XcodeProjectSettings"
+
+static constexpr FStringView kDefaultMacResourcesFolder(TEXTVIEW("Build/Mac/Resources/"));
+static constexpr FStringView kDefaultIOSGeneratedFolder(TEXTVIEW("Build/IOS/UBTGenerated/"));
+
+TSharedRef<IDetailCustomization> FXcodeProjectSettingsDetailsCustomization::MakeInstance()
+{
+    return MakeShareable(new FXcodeProjectSettingsDetailsCustomization);
+}
+
+void FXcodeProjectSettingsDetailsCustomization::CustomizeDetails(IDetailLayoutBuilder& DetailLayout)
+{
+    TemplateMacPlist = DetailLayout.GetProperty(GET_MEMBER_NAME_CHECKED(UXcodeProjectSettings, TemplateMacPlist));
+    TemplateIOSPlist = DetailLayout.GetProperty(GET_MEMBER_NAME_CHECKED(UXcodeProjectSettings, TemplateIOSPlist));
+    PremadeMacEntitlements = DetailLayout.GetProperty(GET_MEMBER_NAME_CHECKED(UXcodeProjectSettings, PremadeMacEntitlements));
+    ShippingEntitlements = DetailLayout.GetProperty(GET_MEMBER_NAME_CHECKED(UXcodeProjectSettings, ShippingSpecificMacEntitlements));
+    
+    IDetailCategoryBuilder& PlistCategory = DetailLayout.EditCategory(TEXT("Info.plist"));
+    PlistCategory.AddCustomRow(LOCTEXT("InfoPlist", "Info.plist"), false)
+    .WholeRowWidget
+    [
+        SNew(SHorizontalBox)
+        + SHorizontalBox::Slot()
+            .Padding(FMargin(0, 5, 0, 10))
+            .AutoWidth()
+        [
+           SNew(SButton)
+           .Text(LOCTEXT("RestoreInfoPlist", "Restore Info.plist to default"))
+           .ToolTipText(LOCTEXT("RestoreInfoPlistTooltip", "Revert to use default templates copied from Engine"))
+           .OnClicked(this, &FXcodeProjectSettingsDetailsCustomization::OnRestorePlistClicked)
+         ]
+    ];
+    
+    IDetailCategoryBuilder& ShipEntitlementCategory = DetailLayout.EditCategory(TEXT("Entitlement"));
+    ShipEntitlementCategory.AddCustomRow(LOCTEXT("Entitlement", "Entitlement"), false)
+    .WholeRowWidget
+    [
+        SNew(SHorizontalBox)
+        + SHorizontalBox::Slot()
+            .Padding(FMargin(0, 5, 0, 10))
+            .AutoWidth()
+        [
+            SNew(SButton)
+            .Text(LOCTEXT("RestoreEntitlements", "Restore entitlements to default"))
+            .ToolTipText(LOCTEXT("RestoreEntitlementsTooltip", "Revert to use default entitlements copied from Engine"))
+            .OnClicked(this, &FXcodeProjectSettingsDetailsCustomization::OnRestoreEntitlementClicked)
+        ]
+    ];
+    
+    DetailLayout.SortCategories([](const TMap<FName, IDetailCategoryBuilder*>& CategoryMap)
+    {
+        for (const TPair<FName, IDetailCategoryBuilder*>& Pair : CategoryMap)
+        {
+            int32 SortOrder = Pair.Value->GetSortOrder();
+            const FName CategoryName = Pair.Key;
+
+            if (CategoryName == "Xcode")
+            {
+                SortOrder = 0;
+            }
+            else if(CategoryName == "Info.plist")
+            {
+                SortOrder = 1;
+            }
+            else if(CategoryName == "Entitlement")
+            {
+                SortOrder = 2;
+            }
+            else
+            {
+                // Unknown category, should explicitly set order
+                ensure(false);
+                SortOrder = 999;
+            }
+
+            Pair.Value->SetSortOrder(SortOrder);
+        }
+
+        return;
+    });
+}
+
+FReply FXcodeProjectSettingsDetailsCustomization::OnRestorePlistClicked()
+{
+    const UXcodeProjectSettings* Settings = GetDefault<UXcodeProjectSettings>();
+    
+    // Copy the default plist template from Engine
+    FText ErrorMessage;
+    if (!SourceControlHelpers::CopyFileUnderSourceControl(FPaths::ProjectDir() + kDefaultMacResourcesFolder + TEXT("Info.Template.plist"),
+                                                          FPaths::EngineDir() + kDefaultMacResourcesFolder + TEXT("Info.Template.plist"),
+                                                          FText::FromString(TEXT("Info.Template.plist")),
+                                                          /*out*/ ErrorMessage))
+    {
+        FNotificationInfo Info(ErrorMessage);
+        Info.ExpireDuration = 3.0f;
+        FSlateNotificationManager::Get().AddNotification(Info);
+    }
+    
+    TemplateMacPlist->SetValueFromFormattedString(TEXT("(FilePath=\"/Game/") + (FString)kDefaultMacResourcesFolder + TEXT("Info.Template.plist\")"));
+    
+    // No need to copy iOS template, it uses generated plist
+    
+    TemplateIOSPlist->SetValueFromFormattedString(TEXT("(FilePath=\"/Game/") + (FString)kDefaultIOSGeneratedFolder + TEXT("Info.Template.plist\")"));
+
+    return FReply::Handled();
+}
+
+FReply FXcodeProjectSettingsDetailsCustomization::OnRestoreEntitlementClicked()
+{
+    const UXcodeProjectSettings* Settings = GetDefault<UXcodeProjectSettings>();
+    
+    // Copy the default entitlememt from Engine
+    FText ErrorMessage;
+    if (!SourceControlHelpers::CopyFileUnderSourceControl(FPaths::ProjectDir() + kDefaultMacResourcesFolder + TEXT("Sandbox.Server.entitlements"),
+                                                          FPaths::EngineDir() + kDefaultMacResourcesFolder + TEXT("Sandbox.Server.entitlements"),
+                                                          FText::FromString(TEXT("Sandbox.Server.entitlements")),
+                                                          /*out*/ ErrorMessage))
+    {
+        FNotificationInfo Info(ErrorMessage);
+        Info.ExpireDuration = 3.0f;
+        FSlateNotificationManager::Get().AddNotification(Info);
+    }
+    
+    PremadeMacEntitlements->SetValueFromFormattedString(TEXT("(FilePath=\"/Game/") + (FString)kDefaultMacResourcesFolder + TEXT("Sandbox.Server.entitlements\")"));
+    
+    // Copy the default entitlememt from Engine
+    if (!SourceControlHelpers::CopyFileUnderSourceControl(FPaths::ProjectDir() + kDefaultMacResourcesFolder + TEXT("Sandbox.NoNet.entitlements"),
+                                                          FPaths::EngineDir() + kDefaultMacResourcesFolder + TEXT("Sandbox.NoNet.entitlements"),
+                                                          FText::FromString(TEXT("Sandbox.NoNet.entitlements")),
+                                                          /*out*/ ErrorMessage))
+    {
+        FNotificationInfo Info(ErrorMessage);
+        Info.ExpireDuration = 3.0f;
+        FSlateNotificationManager::Get().AddNotification(Info);
+    }
+    
+    ShippingEntitlements->SetValueFromFormattedString(TEXT("(FilePath=\"/Game/") + (FString)kDefaultMacResourcesFolder + TEXT("Sandbox.NoNet.entitlements\")"));
+
+    return FReply::Handled();
+}
