fix LOCTEXT_NAMESPACE

diff --git a/Engine/Source/Developer/Mac/MacPlatformEditor/Private/XcodeProjectSettingsDetailsCustomization.cpp b/Engine/Source/Developer/Mac/MacPlatformEditor/Private/XcodeProjectSettingsDetailsCustomization.cpp
--- a/Engine/Source/Developer/Mac/MacPlatformEditor/Private/XcodeProjectSettingsDetailsCustomization.cpp
+++ b/Engine/Source/Developer/Mac/MacPlatformEditor/Private/XcodeProjectSettingsDetailsCustomization.cpp
@@ -152,3 +152,5 @@ FReply FXcodeProjectSettingsDetailsCustomization::OnRestoreEntitlementClicked()
 
     return FReply::Handled();
 }
+
+#undef LOCTEXT_NAMESPACE
