- Fixed up some FilePath stuff with Xcode settings
- Fixed a quotes warning in a .ini

diff --git a/Engine/Source/Developer/Mac/MacTargetPlatform/Classes/XcodeProjectSettings.h b/Engine/Source/Developer/Mac/MacTargetPlatform/Classes/XcodeProjectSettings.h
--- a/Engine/Source/Developer/Mac/MacTargetPlatform/Classes/XcodeProjectSettings.h
+++ b/Engine/Source/Developer/Mac/MacTargetPlatform/Classes/XcodeProjectSettings.h
@@ -133,5 +133,5 @@ public:
 	 * The path to the downloaded .p8 file shared with your team
 	 */
 	UPROPERTY(EditAnywhere, config, Category="Code Signing", meta = (EditCondition="bUseModernXcode && bUseAutomaticCodeSigning && bUseAppStoreConnect", DisplayName = "App Store Connect: Key File"))
-	FString AppStoreConnectKeyPath;
+	FFilePath AppStoreConnectKeyPath;
 };
