PR #11047: Fix typo in xcode project settings

diff --git a/Engine/Source/Developer/Mac/MacTargetPlatform/Classes/XcodeProjectSettings.h b/Engine/Source/Developer/Mac/MacTargetPlatform/Classes/XcodeProjectSettings.h
--- a/Engine/Source/Developer/Mac/MacTargetPlatform/Classes/XcodeProjectSettings.h
+++ b/Engine/Source/Developer/Mac/MacTargetPlatform/Classes/XcodeProjectSettings.h
@@ -95,7 +95,7 @@ public:
 	/**
 	 * The premade entitlement file used for development Mac builds
 	 */
-	UPROPERTY(EditAnywhere, config, Category="Entitlements", meta = (EditCondition="bUseModernXcode", DisplayName = "Nac: Development Entitlements"))
+	UPROPERTY(EditAnywhere, config, Category="Entitlements", meta = (EditCondition="bUseModernXcode", DisplayName = "Mac: Development Entitlements"))
 	FFilePath PremadeMacEntitlements;
 	
 	/**
