Fix Editor Crash when Mac SDK is not installed
Temporary put MacTargetPlatformSettings back to requiring sdk
Also add SetPlatformHiddenFromUI to remove Mac from Platform dropdown

diff --git a/Engine/Source/Developer/Mac/MacTargetPlatformSettings/Private/MacTargetPlatformSettingsModule.cpp b/Engine/Source/Developer/Mac/MacTargetPlatformSettings/Private/MacTargetPlatformSettingsModule.cpp
--- a/Engine/Source/Developer/Mac/MacTargetPlatformSettings/Private/MacTargetPlatformSettingsModule.cpp
+++ b/Engine/Source/Developer/Mac/MacTargetPlatformSettings/Private/MacTargetPlatformSettingsModule.cpp
@@ -39,6 +39,12 @@ public:
 		ITargetPlatformSettings* ClientTP = new TGenericMacTargetPlatformSettings<false, false, true>();
 		TargetPlatforms.Add(ClientTP);
 		PlatformNameToPlatformSettings.Add(FMacPlatformProperties<false, false, true>::PlatformName(), ClientTP);
+
+#if PLATFORM_WINDOWS
+		// we added Mac to Windows so that the Xcode Project settings show up, but we don't
+		// want to see the Mac in the Platforms dropdown
+		FDataDrivenPlatformInfoRegistry::SetPlatformHiddenFromUI("Mac");
+#endif
 	}
 
 	virtual void GetPlatformSettingsMaps(TMap<FString, ITargetPlatformSettings*>& OutMap) override
diff --git a/Engine/Source/Programs/UnrealBuildTool/Platform/Mac/UEBuildMac.cs b/Engine/Source/Programs/UnrealBuildTool/Platform/Mac/UEBuildMac.cs
--- a/Engine/Source/Programs/UnrealBuildTool/Platform/Mac/UEBuildMac.cs
+++ b/Engine/Source/Programs/UnrealBuildTool/Platform/Mac/UEBuildMac.cs
@@ -471,10 +471,9 @@ namespace UnrealBuildTool
 				// because remote IOS building needs the new XcodeProject Settings to show up in the editor, we bring in the Mac bits that expose it
 				if (ModuleName == "Engine")
 				{
-					Rules.DynamicallyLoadedModuleNames.Add("MacTargetPlatformSettings");
 					if (bIsPlatformAvailableForTargetWithSDK)
 					{
-						Rules.DynamicallyLoadedModuleNames.AddAll("MacTargetPlatform", "MacTargetPlatformControls", "MacPlatformEditor");
+						Rules.DynamicallyLoadedModuleNames.AddAll("MacTargetPlatform", "MacTargetPlatformSettings", "MacTargetPlatformControls", "MacPlatformEditor");
 					}
 				}
 			}
