Removing Mac from being marked as a hidden platform as this was temporarily done to fix an editor crash.

#rnx

diff --git a/Engine/Source/Developer/Mac/MacTargetPlatformSettings/Private/MacTargetPlatformSettingsModule.cpp b/Engine/Source/Developer/Mac/MacTargetPlatformSettings/Private/MacTargetPlatformSettingsModule.cpp
--- a/Engine/Source/Developer/Mac/MacTargetPlatformSettings/Private/MacTargetPlatformSettingsModule.cpp
+++ b/Engine/Source/Developer/Mac/MacTargetPlatformSettings/Private/MacTargetPlatformSettingsModule.cpp
@@ -39,12 +39,6 @@ public:
 		ITargetPlatformSettings* ClientTP = new TGenericMacTargetPlatformSettings<false, false, true>();
 		TargetPlatforms.Add(ClientTP);
 		PlatformNameToPlatformSettings.Add(FMacPlatformProperties<false, false, true>::PlatformName(), ClientTP);
-
-#if PLATFORM_WINDOWS
-		// we added Mac to Windows so that the Xcode Project settings show up, but we don't
-		// want to see the Mac in the Platforms dropdown
-		FDataDrivenPlatformInfoRegistry::SetPlatformHiddenFromUI("Mac");
-#endif
 	}
 
 	virtual void GetPlatformSettingsMaps(TMap<FString, ITargetPlatformSettings*>& OutMap) override
