Introduce the MacCookedTarget in the MacTargetPlatform module the same way the WindowsCookedEditorTarget is implemented in WindowsTargetPlatform.

diff --git a/Engine/Source/Developer/Mac/MacTargetPlatform/MacTargetPlatform.Build.cs b/Engine/Source/Developer/Mac/MacTargetPlatform/MacTargetPlatform.Build.cs
--- a/Engine/Source/Developer/Mac/MacTargetPlatform/MacTargetPlatform.Build.cs
+++ b/Engine/Source/Developer/Mac/MacTargetPlatform/MacTargetPlatform.Build.cs
@@ -17,5 +17,10 @@ public class MacTargetPlatform : ModuleRules
 				"MacTargetPlatformControls",
 			}
 		);
+
+		if (Target.bCompileAgainstEngine)
+		{
+			PublicIncludePathModuleNames.Add("Engine");
+        }
 	}
 }
diff --git a/Engine/Source/Developer/Mac/MacTargetPlatformControls/MacTargetPlatformControls.Build.cs b/Engine/Source/Developer/Mac/MacTargetPlatformControls/MacTargetPlatformControls.Build.cs
--- a/Engine/Source/Developer/Mac/MacTargetPlatformControls/MacTargetPlatformControls.Build.cs
+++ b/Engine/Source/Developer/Mac/MacTargetPlatformControls/MacTargetPlatformControls.Build.cs
@@ -30,7 +30,13 @@ public class MacTargetPlatformControls : ModuleRules
 
 		if (Target.bCompileAgainstEngine)
 		{
-			PrivateDependencyModuleNames.Add("Engine");
+            PublicIncludePathModuleNames.Add("Engine");
+            PrivateDependencyModuleNames.AddRange(
+                new string[] {
+                    "Engine",
+                    "CookedEditor",
+                }
+            );
 			PrivateIncludePathModuleNames.Add("TextureCompressor");
 		}
 	}
diff --git a/Engine/Source/Developer/Mac/MacTargetPlatformControls/Private/MacTargetPlatformControlsModule.cpp b/Engine/Source/Developer/Mac/MacTargetPlatformControls/Private/MacTargetPlatformControlsModule.cpp
--- a/Engine/Source/Developer/Mac/MacTargetPlatformControls/Private/MacTargetPlatformControlsModule.cpp
+++ b/Engine/Source/Developer/Mac/MacTargetPlatformControls/Private/MacTargetPlatformControlsModule.cpp
@@ -6,7 +6,9 @@
 #include "Interfaces/ITargetPlatformControlsModule.h"
 #include "Modules/ModuleManager.h"
 #include "IMacTargetPlatformSettingsModule.h"
-
+#if WITH_ENGINE
+#include "CookedEditorTargetPlatformControls.h"
+#endif
 #define LOCTEXT_NAMESPACE "FMacTargetPlatformControlsModule"
 
 /**
@@ -32,6 +34,17 @@ public:
 			TargetPlatforms.Add(EditorTP);
 			TargetPlatforms.Add(ServerTP);
 			TargetPlatforms.Add(ClientTP);
+#if WITH_ENGINE
+			// currently this TP requires the engine for allowing GameDelegates usage
+			bool bSupportCookedEditor;
+			if (GConfig->GetBool(TEXT("CookedEditorSettings"), TEXT("bSupportCookedEditor"), bSupportCookedEditor, GGameIni) && bSupportCookedEditor)
+			{
+				ITargetPlatformControls* CookedEditorTP = new TCookedEditorTargetPlatformControls<FMacEditorTargetPlatformControlsParent>(ModuleSettings->GetCookedEditorPlatformSettings());
+				ITargetPlatformControls* CookedCookerTP = new TCookedCookerTargetPlatformControls<FMacEditorTargetPlatformControlsParent>(ModuleSettings->GetCookedCookerPlatformSettings());
+				TargetPlatforms.Add(CookedEditorTP);
+				TargetPlatforms.Add(CookedCookerTP);
+			}
+#endif
 		}
 	}
 
diff --git a/Engine/Source/Developer/Mac/MacTargetPlatformSettings/MacTargetPlatformSettings.Build.cs b/Engine/Source/Developer/Mac/MacTargetPlatformSettings/MacTargetPlatformSettings.Build.cs
--- a/Engine/Source/Developer/Mac/MacTargetPlatformSettings/MacTargetPlatformSettings.Build.cs
+++ b/Engine/Source/Developer/Mac/MacTargetPlatformSettings/MacTargetPlatformSettings.Build.cs
@@ -26,8 +26,13 @@ public class MacTargetPlatformSettings : ModuleRules
 		);
 
 		if (Target.bCompileAgainstEngine)
-		{
-			PrivateDependencyModuleNames.Add("Engine");
+		{  
+            PublicDependencyModuleNames.AddRange(
+                new string[] {
+                    "CookedEditor",
+                }
+            );
+            PrivateDependencyModuleNames.Add("Engine");
 			PrivateIncludePathModuleNames.Add("TextureCompressor");
 		}
 	}
diff --git a/Engine/Source/Developer/Mac/MacTargetPlatformSettings/Private/MacTargetPlatformSettingsModule.cpp b/Engine/Source/Developer/Mac/MacTargetPlatformSettings/Private/MacTargetPlatformSettingsModule.cpp
--- a/Engine/Source/Developer/Mac/MacTargetPlatformSettings/Private/MacTargetPlatformSettingsModule.cpp
+++ b/Engine/Source/Developer/Mac/MacTargetPlatformSettings/Private/MacTargetPlatformSettingsModule.cpp
@@ -12,6 +12,10 @@
 #include "UObject/WeakObjectPtr.h"
 #include "IMacTargetPlatformSettingsModule.h"
 
+#if WITH_ENGINE
+#include "CookedEditorTargetPlatformSettings.h"
+#endif
+
 #define LOCTEXT_NAMESPACE "FMacTargetPlatformSettingsModule"
 
 /**
@@ -39,6 +43,20 @@ public:
 		ITargetPlatformSettings* ClientTP = new TGenericMacTargetPlatformSettings<false, false, true>();
 		TargetPlatforms.Add(ClientTP);
 		PlatformNameToPlatformSettings.Add(FMacPlatformProperties<false, false, true>::PlatformName(), ClientTP);
+        
+#if WITH_ENGINE && COOKEDEDITOR_WITH_MACTARGETPLATFORM
+        // currently this TP requires the engine for allowing GameDelegates usage
+        bool bSupportCookedEditor;
+        if (GConfig->GetBool(TEXT("CookedEditorSettings"), TEXT("bSupportCookedEditor"), bSupportCookedEditor, GGameIni) && bSupportCookedEditor)
+        {
+			ITargetPlatformSettings* CookedEditorTP = new TCookedEditorTargetPlatformSettings<FMacEditorTargetPlatformSettingsParent>();
+			ITargetPlatformSettings* CookedCookerTP = new TCookedCookerTargetPlatformSettings<FMacEditorTargetPlatformSettingsParent>();
+			PlatformSettingsCookedEditor = CookedEditorTP;
+			PlatformSettingsCookedCooker = CookedCookerTP;
+			TargetPlatforms.Add(CookedEditorTP);
+			TargetPlatforms.Add(CookedCookerTP);
+        }
+#endif
 	}
 
 	virtual void GetPlatformSettingsMaps(TMap<FString, ITargetPlatformSettings*>& OutMap) override
@@ -46,6 +64,16 @@ public:
 		OutMap = PlatformNameToPlatformSettings;
 	}
 
+	virtual ITargetPlatformSettings* GetCookedEditorPlatformSettings() override
+	{
+		return PlatformSettingsCookedEditor;
+	}
+
+	virtual ITargetPlatformSettings* GetCookedCookerPlatformSettings() override
+	{
+		return PlatformSettingsCookedCooker;
+	}
+
 public:
 
 	// Begin IModuleInterface interface
@@ -131,6 +159,8 @@ public:
 
 private:
 	TMap<FString, ITargetPlatformSettings*> PlatformNameToPlatformSettings;
+	ITargetPlatformSettings* PlatformSettingsCookedEditor;
+	ITargetPlatformSettings* PlatformSettingsCookedCooker;
 
 	// Holds the target settings.
 	UMacTargetSettings* TargetSettings;
diff --git a/Engine/Source/Developer/Mac/MacTargetPlatformSettings/Public/IMacTargetPlatformSettingsModule.h b/Engine/Source/Developer/Mac/MacTargetPlatformSettings/Public/IMacTargetPlatformSettingsModule.h
--- a/Engine/Source/Developer/Mac/MacTargetPlatformSettings/Public/IMacTargetPlatformSettingsModule.h
+++ b/Engine/Source/Developer/Mac/MacTargetPlatformSettings/Public/IMacTargetPlatformSettingsModule.h
@@ -13,4 +13,6 @@ class IMacTargetPlatformSettingsModule : public ITargetPlatformSettingsModule
 {
 public:
 	virtual void GetPlatformSettingsMaps(TMap<FString, ITargetPlatformSettings*>& OutMap) = 0;
+	virtual ITargetPlatformSettings* GetCookedEditorPlatformSettings() = 0;
+	virtual ITargetPlatformSettings* GetCookedCookerPlatformSettings() = 0;
 };
