- Fix for Mac failing to startup if TargettedRHI is missing from config, now displays error and falls back to SM5
- Changed MacTargetSettings so you cannot unselect all target platforms

diff --git a/Engine/Source/Developer/Mac/MacTargetPlatformSettings/Classes/MacTargetSettings.h b/Engine/Source/Developer/Mac/MacTargetPlatformSettings/Classes/MacTargetSettings.h
--- a/Engine/Source/Developer/Mac/MacTargetPlatformSettings/Classes/MacTargetSettings.h
+++ b/Engine/Source/Developer/Mac/MacTargetPlatformSettings/Classes/MacTargetSettings.h
@@ -171,4 +171,14 @@ public:
 	{
 		InOutSectionName = TEXT("/Script/MacTargetPlatform.MacTargetSettings");
 	}
+	
+#if WITH_EDITOR
+	virtual void PostEditChangeProperty(struct FPropertyChangedEvent& PropertyChangedEvent) override
+	{
+		if(TargetedRHIs.Num() == 0)
+		{		
+			TargetedRHIs.Add(TEXT("SF_METAL_SM5"));
+		}
+	}
+#endif
 };
diff --git a/Engine/Source/Runtime/RHI/Private/Apple/AppleDynamicRHI.cpp b/Engine/Source/Runtime/RHI/Private/Apple/AppleDynamicRHI.cpp
--- a/Engine/Source/Runtime/RHI/Private/Apple/AppleDynamicRHI.cpp
+++ b/Engine/Source/Runtime/RHI/Private/Apple/AppleDynamicRHI.cpp
@@ -105,6 +105,11 @@ static void ComputeRequestedFeatureLevel(const FAppleDynamicRHIOptions& Options,
                 }
             }
 		}
+		else
+		{
+			FMessageDialog::Open(EAppMsgType::Ok, LOCTEXT("RHIOptionsError", "No Targeted RHI is set for this project, defaulting to SM5"));
+			RequestedFeatureLevel = ERHIFeatureLevel::SM5;
+		}
 	}
 	else
 	{
