Expose IOS entitlement paths.

#rn Expose IOS entitlements settings in Modern Xcode project settings.

diff --git a/Engine/Source/Developer/Mac/MacPlatformEditor/Private/XcodeProjectSettingsDetailsCustomization.cpp b/Engine/Source/Developer/Mac/MacPlatformEditor/Private/XcodeProjectSettingsDetailsCustomization.cpp
--- a/Engine/Source/Developer/Mac/MacPlatformEditor/Private/XcodeProjectSettingsDetailsCustomization.cpp
+++ b/Engine/Source/Developer/Mac/MacPlatformEditor/Private/XcodeProjectSettingsDetailsCustomization.cpp
@@ -27,8 +27,10 @@ void FXcodeProjectSettingsDetailsCustomization::CustomizeDetails(IDetailLayoutBu
     TemplateMacPlist = DetailLayout.GetProperty(GET_MEMBER_NAME_CHECKED(UXcodeProjectSettings, TemplateMacPlist));
     TemplateIOSPlist = DetailLayout.GetProperty(GET_MEMBER_NAME_CHECKED(UXcodeProjectSettings, TemplateIOSPlist));
     PremadeMacEntitlements = DetailLayout.GetProperty(GET_MEMBER_NAME_CHECKED(UXcodeProjectSettings, PremadeMacEntitlements));
-    ShippingEntitlements = DetailLayout.GetProperty(GET_MEMBER_NAME_CHECKED(UXcodeProjectSettings, ShippingSpecificMacEntitlements));
-    
+    ShippingMacEntitlements = DetailLayout.GetProperty(GET_MEMBER_NAME_CHECKED(UXcodeProjectSettings, ShippingSpecificMacEntitlements));
+    PremadeIOSEntitlements = DetailLayout.GetProperty(GET_MEMBER_NAME_CHECKED(UXcodeProjectSettings, PremadeIOSEntitlements));
+    ShippingIOSEntitlements = DetailLayout.GetProperty(GET_MEMBER_NAME_CHECKED(UXcodeProjectSettings, ShippingSpecificIOSEntitlements));
+
     IDetailCategoryBuilder& PlistCategory = DetailLayout.EditCategory(TEXT("Plist Files"));
     PlistCategory.AddCustomRow(LOCTEXT("InfoPlist", "Info.plist"), false)
     .WholeRowWidget
@@ -156,8 +158,11 @@ FReply FXcodeProjectSettingsDetailsCustomization::OnRestoreEntitlementClicked()
         FSlateNotificationManager::Get().AddNotification(Info);
     }
     
-    ShippingEntitlements->SetValueFromFormattedString(TEXT("(FilePath=\"/Game/") + (FString)kDefaultMacResourcesFolder + TEXT("Sandbox.NoNet.entitlements\")"));
-
+    ShippingMacEntitlements->SetValueFromFormattedString(TEXT("(FilePath=\"/Game/") + (FString)kDefaultMacResourcesFolder + TEXT("Sandbox.NoNet.entitlements\")"));
+	
+	PremadeIOSEntitlements->SetValue(TEXT(""));
+	ShippingIOSEntitlements->SetValue(TEXT(""));
+	
     return FReply::Handled();
 }
 
diff --git a/Engine/Source/Developer/Mac/MacPlatformEditor/Public/XcodeProjectSettingsDetailsCustomization.h b/Engine/Source/Developer/Mac/MacPlatformEditor/Public/XcodeProjectSettingsDetailsCustomization.h
--- a/Engine/Source/Developer/Mac/MacPlatformEditor/Public/XcodeProjectSettingsDetailsCustomization.h
+++ b/Engine/Source/Developer/Mac/MacPlatformEditor/Public/XcodeProjectSettingsDetailsCustomization.h
@@ -29,5 +29,7 @@ public:
     TSharedPtr<IPropertyHandle> TemplateMacPlist;
     TSharedPtr<IPropertyHandle> TemplateIOSPlist;
     TSharedPtr<IPropertyHandle> PremadeMacEntitlements;
-    TSharedPtr<IPropertyHandle> ShippingEntitlements;
+    TSharedPtr<IPropertyHandle> ShippingMacEntitlements;
+    TSharedPtr<IPropertyHandle> PremadeIOSEntitlements;
+    TSharedPtr<IPropertyHandle> ShippingIOSEntitlements;
 };
diff --git a/Engine/Source/Developer/Mac/MacTargetPlatformSettings/Classes/XcodeProjectSettings.h b/Engine/Source/Developer/Mac/MacTargetPlatformSettings/Classes/XcodeProjectSettings.h
--- a/Engine/Source/Developer/Mac/MacTargetPlatformSettings/Classes/XcodeProjectSettings.h
+++ b/Engine/Source/Developer/Mac/MacTargetPlatformSettings/Classes/XcodeProjectSettings.h
@@ -104,6 +104,18 @@ public:
 	UPROPERTY(EditAnywhere, config, Category="Entitlements", meta = (EditCondition="bUseModernXcode", DisplayName = "Mac: Shipping Entitlements"))
 	FFilePath ShippingSpecificMacEntitlements;
 	
+	/**
+	 * The premade entitlement file used for development IOS builds
+	 */
+	UPROPERTY(EditAnywhere, config, Category="Entitlements", meta = (EditCondition="bUseModernXcode", DisplayName = "IOS: Development Entitlements"))
+	FFilePath PremadeIOSEntitlements;
+	
+	/**
+	 * The premade entitlement file used for shipping IOS builds
+	 */
+	UPROPERTY(EditAnywhere, config, Category="Entitlements", meta = (EditCondition="bUseModernXcode", DisplayName = "IOS: Shipping Entitlements"))
+	FFilePath ShippingSpecificIOSEntitlements;
+	
 	/**
 	 * Enable native Xcode code signing
 	 */
@@ -180,4 +192,4 @@ public:
 	{
 		InOutSectionName = TEXT("/Script/MacTargetPlatform.XcodeProjectSettings");
 	}
-};
\ No newline at end of file
+};
