Setup default entitlements for iOS under Modern Xcode. Default entitlements including the Multicast Networking are required for IOS due to the fact that the UdpMessaging plugin is enabled my default via -Messaging. Failing to provide this entitlement causes the UdbMessaging plugin to fail at runtime and get into a permanent retry state causing periodic micro hitches at 0.5-1s intervals under stock settings. Additionally, some drive by cleanup commonizing the code that restores plist/entitlement setting and resources to default.

diff --git a/Engine/Config/BaseEngine.ini b/Engine/Config/BaseEngine.ini
--- a/Engine/Config/BaseEngine.ini
+++ b/Engine/Config/BaseEngine.ini
@@ -3387,6 +3387,8 @@ TemplateMacPlist=(FilePath="/Game/Build/Mac/Resources/Info.Template.plist")
 TemplateIOSPlist=(FilePath="/Game/Build/IOS/UBTGenerated/Info.Template.plist")
 PremadeMacEntitlements=(FilePath="/Game/Build/Mac/Resources/Sandbox.Server.entitlements")
 ShippingSpecificMacEntitlements=(FilePath="/Game/Build/Mac/Resources/Sandbox.NoNet.entitlements")
+PremadeIOSEntitlements=(FilePath="/Game/Build/IOS/Resources/Development.entitlements")
+ShippingSpecificIOSEntitlements=(FilePath="/Game/Build/IOS/Resources/Shipping.entitlements")
 ; Programs generally don't need entitlements. Override in a particular program's .ini to true to use the above Premade entitlements files
 bUseEntitlementsForPrograms=false
 ; The following are used only with Manual codesigning
diff --git a/Engine/Source/Developer/Mac/MacPlatformEditor/Private/XcodeProjectSettingsDetailsCustomization.cpp b/Engine/Source/Developer/Mac/MacPlatformEditor/Private/XcodeProjectSettingsDetailsCustomization.cpp
--- a/Engine/Source/Developer/Mac/MacPlatformEditor/Private/XcodeProjectSettingsDetailsCustomization.cpp
+++ b/Engine/Source/Developer/Mac/MacPlatformEditor/Private/XcodeProjectSettingsDetailsCustomization.cpp
@@ -15,6 +15,7 @@
 #define LOCTEXT_NAMESPACE "XcodeProjectSettings"
 
 static constexpr FStringView kDefaultMacResourcesFolder(TEXTVIEW("Build/Mac/Resources/"));
+static constexpr FStringView kDefaultIOSResourcesFolder(TEXTVIEW("Build/IOS/Resources/"));
 static constexpr FStringView kDefaultIOSGeneratedFolder(TEXTVIEW("Build/IOS/UBTGenerated/"));
 
 TSharedRef<IDetailCustomization> FXcodeProjectSettingsDetailsCustomization::MakeInstance()
@@ -104,26 +105,31 @@ void FXcodeProjectSettingsDetailsCustomization::CustomizeDetails(IDetailLayoutBu
     });
 }
 
+void FXcodeProjectSettingsDetailsCustomization::RestoreDefault(const FStringView& SubFolder, const TCHAR* Filename, TSharedPtr<IPropertyHandle> Property)
+{
+	FText ErrorMessage;
+	// Copy the default from Engine
+	if (!SourceControlHelpers::CopyFileUnderSourceControl(FPaths::ProjectDir() + SubFolder + Filename,
+														  FPaths::EngineDir() + SubFolder + Filename,
+														  FText::FromString(Filename),
+														  /*out*/ ErrorMessage))
+	{
+		FNotificationInfo Info(ErrorMessage);
+		Info.ExpireDuration = 3.0f;
+		FSlateNotificationManager::Get().AddNotification(Info);
+	}
+	
+	Property->SetValueFromFormattedString(TEXT("(FilePath=\"/Game/") + (FString)SubFolder + Filename + TEXT("\")"));
+}
+
+
 FReply FXcodeProjectSettingsDetailsCustomization::OnRestorePlistClicked()
 {
     const UXcodeProjectSettings* Settings = GetDefault<UXcodeProjectSettings>();
     
-    // Copy the default plist template from Engine
-    FText ErrorMessage;
-    if (!SourceControlHelpers::CopyFileUnderSourceControl(FPaths::ProjectDir() + kDefaultMacResourcesFolder + TEXT("Info.Template.plist"),
-                                                          FPaths::EngineDir() + kDefaultMacResourcesFolder + TEXT("Info.Template.plist"),
-                                                          FText::FromString(TEXT("Info.Template.plist")),
-                                                          /*out*/ ErrorMessage))
-    {
-        FNotificationInfo Info(ErrorMessage);
-        Info.ExpireDuration = 3.0f;
-        FSlateNotificationManager::Get().AddNotification(Info);
-    }
-    
-    TemplateMacPlist->SetValueFromFormattedString(TEXT("(FilePath=\"/Game/") + (FString)kDefaultMacResourcesFolder + TEXT("Info.Template.plist\")"));
+	RestoreDefault(kDefaultMacResourcesFolder, TEXT("Info.Template.plist"), TemplateMacPlist);
     
     // No need to copy iOS template, it uses generated plist
-    
     TemplateIOSPlist->SetValueFromFormattedString(TEXT("(FilePath=\"/Game/") + (FString)kDefaultIOSGeneratedFolder + TEXT("Info.Template.plist\")"));
 
     return FReply::Handled();
@@ -133,35 +139,11 @@ FReply FXcodeProjectSettingsDetailsCustomization::OnRestoreEntitlementClicked()
 {
     const UXcodeProjectSettings* Settings = GetDefault<UXcodeProjectSettings>();
     
-    // Copy the default entitlememt from Engine
-    FText ErrorMessage;
-    if (!SourceControlHelpers::CopyFileUnderSourceControl(FPaths::ProjectDir() + kDefaultMacResourcesFolder + TEXT("Sandbox.Server.entitlements"),
-                                                          FPaths::EngineDir() + kDefaultMacResourcesFolder + TEXT("Sandbox.Server.entitlements"),
-                                                          FText::FromString(TEXT("Sandbox.Server.entitlements")),
-                                                          /*out*/ ErrorMessage))
-    {
-        FNotificationInfo Info(ErrorMessage);
-        Info.ExpireDuration = 3.0f;
-        FSlateNotificationManager::Get().AddNotification(Info);
-    }
-    
-    PremadeMacEntitlements->SetValueFromFormattedString(TEXT("(FilePath=\"/Game/") + (FString)kDefaultMacResourcesFolder + TEXT("Sandbox.Server.entitlements\")"));
-    
-    // Copy the default entitlememt from Engine
-    if (!SourceControlHelpers::CopyFileUnderSourceControl(FPaths::ProjectDir() + kDefaultMacResourcesFolder + TEXT("Sandbox.NoNet.entitlements"),
-                                                          FPaths::EngineDir() + kDefaultMacResourcesFolder + TEXT("Sandbox.NoNet.entitlements"),
-                                                          FText::FromString(TEXT("Sandbox.NoNet.entitlements")),
-                                                          /*out*/ ErrorMessage))
-    {
-        FNotificationInfo Info(ErrorMessage);
-        Info.ExpireDuration = 3.0f;
-        FSlateNotificationManager::Get().AddNotification(Info);
-    }
-    
-    ShippingMacEntitlements->SetValueFromFormattedString(TEXT("(FilePath=\"/Game/") + (FString)kDefaultMacResourcesFolder + TEXT("Sandbox.NoNet.entitlements\")"));
+    RestoreDefault(kDefaultMacResourcesFolder, TEXT("Sandbox.Server.entitlements"), PremadeMacEntitlements);
+	RestoreDefault(kDefaultMacResourcesFolder, TEXT("Sandbox.NoNet.entitlements"), ShippingMacEntitlements);
 	
-	PremadeIOSEntitlements->SetValue(TEXT(""));
-	ShippingIOSEntitlements->SetValue(TEXT(""));
+	RestoreDefault(kDefaultIOSResourcesFolder, TEXT("Development.entitlements"), PremadeIOSEntitlements);
+	RestoreDefault(kDefaultIOSResourcesFolder, TEXT("Shipping.entitlements"), ShippingIOSEntitlements);
 	
     return FReply::Handled();
 }
diff --git a/Engine/Source/Developer/Mac/MacPlatformEditor/Public/XcodeProjectSettingsDetailsCustomization.h b/Engine/Source/Developer/Mac/MacPlatformEditor/Public/XcodeProjectSettingsDetailsCustomization.h
--- a/Engine/Source/Developer/Mac/MacPlatformEditor/Public/XcodeProjectSettingsDetailsCustomization.h
+++ b/Engine/Source/Developer/Mac/MacPlatformEditor/Public/XcodeProjectSettingsDetailsCustomization.h
@@ -32,4 +32,7 @@ public:
     TSharedPtr<IPropertyHandle> ShippingMacEntitlements;
     TSharedPtr<IPropertyHandle> PremadeIOSEntitlements;
     TSharedPtr<IPropertyHandle> ShippingIOSEntitlements;
+	
+private:
+	static void RestoreDefault(const FStringView& SubFolder, const TCHAR* Filename, TSharedPtr<IPropertyHandle> Property);
 };
diff --git a/Engine/Source/Programs/UnrealBuildTool/ProjectFiles/Xcode/XcodeProject.cs b/Engine/Source/Programs/UnrealBuildTool/ProjectFiles/Xcode/XcodeProject.cs
--- a/Engine/Source/Programs/UnrealBuildTool/ProjectFiles/Xcode/XcodeProject.cs
+++ b/Engine/Source/Programs/UnrealBuildTool/ProjectFiles/Xcode/XcodeProject.cs
@@ -228,9 +228,10 @@ namespace UnrealBuildTool.XcodeProjectXcconfig
 
 		public Metadata(DirectoryReference ProductDirectory, DirectoryReference XcodeProject, ConfigHierarchy Ini, bool bSupportsMac, bool bSupportsIOSOrTVOS, ILogger Logger)
 		{
-			DirectoryReference ResourceFolder = DirectoryReference.Combine(Unreal.EngineDirectory, "Build/Mac/Resources/");
 			if (bSupportsMac)
 			{
+				DirectoryReference ResourceFolder = DirectoryReference.Combine(Unreal.EngineDirectory, "Build/Mac/Resources/");
+	
 				// All editor use template plist which is just a copy of default info.plist, user should not need to modify or change this
 				// TEMP: currently cooked editor (e.g. QAGCookedEditor) uses this hardcoded file too, we should find a way to expose this in Game/Config/DefaultEngine.ini
 				PlistFiles[MetadataPlatform.MacEditor] = new MetadataItem(ProductDirectory, XcodeProject, Ini,
@@ -289,6 +290,8 @@ namespace UnrealBuildTool.XcodeProjectXcconfig
 			}
 			if (bSupportsIOSOrTVOS)
 			{
+				DirectoryReference ResourceFolder = DirectoryReference.Combine(Unreal.EngineDirectory, "Build/IOS/Resources/");
+
 				// Try Premade first
 				PlistFiles[MetadataPlatform.IOS] = new MetadataItem(ProductDirectory, XcodeProject, Ini,
 					"PremadeIOSPlist",
