Revert back to IOS being entitlements free by default. Clean up UI names to make clear that Shipping Entitlements are overrides for default entitlements for both Mac and IOS as the existing UI mislead into thinking both properties needed to be populated independently.

diff --git a/Engine/Config/BaseEngine.ini b/Engine/Config/BaseEngine.ini
--- a/Engine/Config/BaseEngine.ini
+++ b/Engine/Config/BaseEngine.ini
@@ -3417,7 +3417,6 @@ TemplateMacPlist=(FilePath="/Game/Build/Mac/Resources/Info.Template.plist")
 TemplateIOSPlist=(FilePath="/Game/Build/IOS/UBTGenerated/Info.Template.plist")
 PremadeMacEntitlements=(FilePath="/Game/Build/Mac/Resources/Sandbox.Server.entitlements")
 ShippingSpecificMacEntitlements=(FilePath="/Game/Build/Mac/Resources/Sandbox.NoNet.entitlements")
-PremadeIOSEntitlements=(FilePath="/Game/Build/IOS/Resources/Development.entitlements")
 ; Programs generally don't need entitlements. Override in a particular program's .ini to true to use the above Premade entitlements files
 bUseEntitlementsForPrograms=false
 ; The following are used only with Manual codesigning
diff --git a/Engine/Source/Developer/Mac/MacPlatformEditor/Private/XcodeProjectSettingsDetailsCustomization.cpp b/Engine/Source/Developer/Mac/MacPlatformEditor/Private/XcodeProjectSettingsDetailsCustomization.cpp
--- a/Engine/Source/Developer/Mac/MacPlatformEditor/Private/XcodeProjectSettingsDetailsCustomization.cpp
+++ b/Engine/Source/Developer/Mac/MacPlatformEditor/Private/XcodeProjectSettingsDetailsCustomization.cpp
@@ -124,9 +124,7 @@ void FXcodeProjectSettingsDetailsCustomization::RestoreDefault(const FStringView
 
 
 FReply FXcodeProjectSettingsDetailsCustomization::OnRestorePlistClicked()
-{
-    const UXcodeProjectSettings* Settings = GetDefault<UXcodeProjectSettings>();
-    
+{    
 	RestoreDefault(kDefaultMacResourcesFolder, TEXT("Info.Template.plist"), TemplateMacPlist);
     
     // No need to copy iOS template, it uses generated plist
@@ -142,8 +140,8 @@ FReply FXcodeProjectSettingsDetailsCustomization::OnRestoreEntitlementClicked()
     RestoreDefault(kDefaultMacResourcesFolder, TEXT("Sandbox.Server.entitlements"), PremadeMacEntitlements);
 	RestoreDefault(kDefaultMacResourcesFolder, TEXT("Sandbox.NoNet.entitlements"), ShippingMacEntitlements);
 	
-	RestoreDefault(kDefaultIOSResourcesFolder, TEXT("Development.entitlements"), PremadeIOSEntitlements);
-	RestoreDefault(kDefaultIOSResourcesFolder, TEXT("Shipping.entitlements"), ShippingIOSEntitlements);
+	PremadeIOSEntitlements->SetValueFromFormattedString(TEXT("(FilePath=\"\")"));
+	ShippingIOSEntitlements->SetValueFromFormattedString(TEXT("(FilePath=\"\")"));
 	
     return FReply::Handled();
 }
diff --git a/Engine/Source/Developer/Mac/MacTargetPlatformSettings/Classes/XcodeProjectSettings.h b/Engine/Source/Developer/Mac/MacTargetPlatformSettings/Classes/XcodeProjectSettings.h
--- a/Engine/Source/Developer/Mac/MacTargetPlatformSettings/Classes/XcodeProjectSettings.h
+++ b/Engine/Source/Developer/Mac/MacTargetPlatformSettings/Classes/XcodeProjectSettings.h
@@ -95,25 +95,25 @@ public:
 	/**
 	 * The premade entitlement file used for development Mac builds
 	 */
-	UPROPERTY(EditAnywhere, config, Category="Entitlements", meta = (EditCondition="bUseModernXcode", DisplayName = "Mac: Development Entitlements"))
+	UPROPERTY(EditAnywhere, config, Category="Entitlements", meta = (EditCondition="bUseModernXcode", DisplayName = "Mac: Entitlements"))
 	FFilePath PremadeMacEntitlements;
 	
 	/**
 	 * The premade entitlement file used for shipping Mac builds
 	 */
-	UPROPERTY(EditAnywhere, config, Category="Entitlements", meta = (EditCondition="bUseModernXcode", DisplayName = "Mac: Shipping Entitlements"))
+	UPROPERTY(EditAnywhere, config, Category="Entitlements", meta = (EditCondition="bUseModernXcode", DisplayName = "Mac:  Entitlements (Shipping config override)"))
 	FFilePath ShippingSpecificMacEntitlements;
 	
 	/**
 	 * The premade entitlement file used for development IOS builds
 	 */
-	UPROPERTY(EditAnywhere, config, Category="Entitlements", meta = (EditCondition="bUseModernXcode", DisplayName = "IOS: Development Entitlements"))
+	UPROPERTY(EditAnywhere, config, Category="Entitlements", meta = (EditCondition="bUseModernXcode", DisplayName = "IOS: Entitlements"))
 	FFilePath PremadeIOSEntitlements;
 	
 	/**
 	 * The premade entitlement file used for shipping IOS builds
 	 */
-	UPROPERTY(EditAnywhere, config, Category="Entitlements", meta = (EditCondition="bUseModernXcode", DisplayName = "IOS: Shipping Entitlements"))
+	UPROPERTY(EditAnywhere, config, Category="Entitlements", meta = (EditCondition="bUseModernXcode", DisplayName = "IOS: Entitlements (Shipping config override)"))
 	FFilePath ShippingSpecificIOSEntitlements;
 	
 	/**
