Metal - Refactor bindless to support new FeatureLevel & Platforms and include a per-project experimental Mac option to enable SM5 Bindless

#tests Testing on Mac & iOS

diff --git a/Engine/Config/Mac/BaseMacEngine.ini b/Engine/Config/Mac/BaseMacEngine.ini
--- a/Engine/Config/Mac/BaseMacEngine.ini
+++ b/Engine/Config/Mac/BaseMacEngine.ini
@@ -37,6 +37,12 @@ r.HairStrands.Visibility.MSAA.SamplePerPixel=1
 r.SceneCulling.ExplicitCellBounds=1
 r.GPUStatsEnabled=0
 
+[ShaderPlatformConfig METAL_ES3_1]
+BindlessConfiguration=Disabled
+
+[ShaderPlatformConfig METAL_SM5]
+BindlessConfiguration=Disabled
+
 [ShaderPlatformConfig METAL_SM6]
 BindlessConfiguration=All
 
diff --git a/Engine/Config/Mac/DataDrivenPlatformInfo.ini b/Engine/Config/Mac/DataDrivenPlatformInfo.ini
--- a/Engine/Config/Mac/DataDrivenPlatformInfo.ini
+++ b/Engine/Config/Mac/DataDrivenPlatformInfo.ini
@@ -54,12 +54,14 @@ bSupportsRayTracing = false
 bSupportsRayTracingShaders = false
 bSupportsInlineRayTracing = false
 bSupportsRayTracingIndirectInstanceData = false
+// Reenable bSupportsHairStrandGeometry when bindless will be enabled by default
 bSupportsHairStrandGeometry = false
 bSupportsNanite=false
 bSupportsUInt64ImageAtomics=false
 bSupportsNNEShaders=false
 bSupportsIndependentSamplers=true
 bSupportsComputeFramework = true
+bSupportsBindless=true
 FriendlyName=LOCTEXT("FriendlyShaderPlatformName_MacSM5", "Mac Metal SM5")
 
 [ShaderPlatform METAL_SM6]
@@ -120,6 +122,7 @@ bSupportsShaderPipelines = false
 bSupportsDualSourceBlending = true
 bSupportsIndependentSamplers=true
 bSupportsComputeFramework = true
+bSupportsBindless=true
 FriendlyName=LOCTEXT("FriendlyShaderPlatformName_MacMobile", "Mac Metal Mobile")
 
 [PreviewPlatform METAL_SM5]
diff --git a/Engine/Shaders/Public/Platform/Metal/MetalCommon.ush b/Engine/Shaders/Public/Platform/Metal/MetalCommon.ush
--- a/Engine/Shaders/Public/Platform/Metal/MetalCommon.ush
+++ b/Engine/Shaders/Public/Platform/Metal/MetalCommon.ush
@@ -7,7 +7,7 @@
 #pragma once
 
 // Update this GUID to invalidate shader recompilation for only Metal shaders
-#pragma message("UESHADERMETADATA_VERSION 99D4B0D2-064D-4B5C-9292-57A1ED268BB9")
+#pragma message("UESHADERMETADATA_VERSION 80E79B41-0792-4136-8D8E-BB1BA48C008D")
 
 #if VERTEXSHADER && HULLSHADER
 #ifndef TESSELLATIONSHADER
@@ -81,6 +81,12 @@
 	#define SubpassFetchR_5() (0.0)
 #endif //IOS && PIXELSHADER
 
+#if PLATFORM_SUPPORTS_BINDLESS
+	#define GetSRVFromHeap(Type, Index) ResourceDescriptorHeap[Index]
+	#define GetUAVFromHeap(Type, Index) ResourceDescriptorHeap[Index]
+	#define GetSamplerFromHeap(Type, Index) SamplerDescriptorHeap[Index]
+#endif
+
 #if METAL_SM6_PROFILE == 1
     #define COMPILER_SUPPORTS_ULONG_TYPES 1
     #define COMPILER_SUPPORTS_UINT64_IMAGE_ATOMICS 1
@@ -111,14 +117,9 @@
 		}
 
     #endif
-    
-    #if PLATFORM_SUPPORTS_BINDLESS
-        #define GetSRVFromHeap(Type, Index) ResourceDescriptorHeap[Index]
-        #define GetUAVFromHeap(Type, Index) ResourceDescriptorHeap[Index]
-        #define GetSamplerFromHeap(Type, Index) SamplerDescriptorHeap[Index]
-    #endif
 #endif
 
+
 #if PLATFORM_SUPPORTS_MESH_SHADERS_TIER0
     #define MESH_SHADER_TRIANGLE_ATTRIBUTES(InThreadsX) [numthreads(InThreadsX, 1, 1)][outputtopology("triangle")]
     #define MESH_SHADER_LINE_ATTRIBUTES(InThreadsX) [numthreads(InThreadsX, 1, 1)][outputtopology("line")]
diff --git a/Engine/Source/Developer/Apple/MetalShaderFormat/Private/MetalShaderCompiler.cpp b/Engine/Source/Developer/Apple/MetalShaderFormat/Private/MetalShaderCompiler.cpp
--- a/Engine/Source/Developer/Apple/MetalShaderFormat/Private/MetalShaderCompiler.cpp
+++ b/Engine/Source/Developer/Apple/MetalShaderFormat/Private/MetalShaderCompiler.cpp
@@ -418,7 +418,7 @@ void BuildMetalShaderOutput(
 
 	bool bUseMetalShaderConverter = false;
 #if UE_METAL_USE_METAL_SHADER_CONVERTER
-	bUseMetalShaderConverter = ShaderInput.Target.GetPlatform() == EShaderPlatform::SP_METAL_SM6 && bBindlessEnabled;
+	bUseMetalShaderConverter = bBindlessEnabled;
 #endif
 	
 	// Packed Uniform Buffers
@@ -1097,7 +1097,7 @@ void CompileMetalShader(const FShaderCompilerInput& Input, const FShaderPreproce
 
 
 #if UE_METAL_USE_METAL_SHADER_CONVERTER
-	if (Input.IsBindlessEnabled() && Input.ShaderFormat == NAME_SF_METAL_SM6)
+	if (Input.IsBindlessEnabled())
 	{
 		FMetalCompileShaderMSC::DoCompileMetalShader(Input, Output, PreprocessedSource, VersionEnum, Semantics, MaxUnrollLoops, (EShaderFrequency)Input.Target.Frequency, bDumpDebugInfo, Standard, MinOSVersion);
 	}
diff --git a/Engine/Source/Developer/Apple/MetalShaderFormat/Private/MetalShaderFormat.cpp b/Engine/Source/Developer/Apple/MetalShaderFormat/Private/MetalShaderFormat.cpp
--- a/Engine/Source/Developer/Apple/MetalShaderFormat/Private/MetalShaderFormat.cpp
+++ b/Engine/Source/Developer/Apple/MetalShaderFormat/Private/MetalShaderFormat.cpp
@@ -180,8 +180,11 @@ public:
 
 		check(ShaderFormatName == NAME_SF_METAL_ES3_1_IOS || ShaderFormatName == NAME_SF_METAL_SM5_IOS || ShaderFormatName == NAME_SF_METAL_ES3_1_TVOS || ShaderFormatName == NAME_SF_METAL_SM5_TVOS || ShaderFormatName == NAME_SF_METAL_SM5 || ShaderFormatName == NAME_SF_METAL_SM6 || ShaderFormatName == NAME_SF_METAL_SIM || ShaderFormatName == NAME_SF_METAL_ES3_1);
 
-		// SM6 needs a lower limit of shaders per library as the packing process takes a significant amount of RAM
-		if(ShaderFormatName == NAME_SF_METAL_SM6)
+		// Bindless needs a lower limit of shaders per library as the packing process takes a significant amount of RAM
+		const EShaderPlatform Platform = FMetalCompilerToolchain::MetalShaderFormatToLegacyShaderPlatform(ShaderFormatName);
+		if(FDataDrivenShaderPlatformInfo::GetSupportsBindless(Platform)
+		   && UE::ShaderCompiler::GetBindlessConfiguration(Platform) != ERHIBindlessConfiguration::Disabled
+		   && ShouldGenerateShaderSymbolsInfo(ShaderFormatName))
 		{
 			NumShadersPerLibrary = 1000;
 		}
diff --git a/Engine/Source/Developer/Mac/MacPlatformEditor/MacPlatformEditor.Build.cs b/Engine/Source/Developer/Mac/MacPlatformEditor/MacPlatformEditor.Build.cs
--- a/Engine/Source/Developer/Mac/MacPlatformEditor/MacPlatformEditor.Build.cs
+++ b/Engine/Source/Developer/Mac/MacPlatformEditor/MacPlatformEditor.Build.cs
@@ -16,7 +16,6 @@ public class MacPlatformEditor : ModuleRules
 				"MainFrame",
 				"Slate",
 				"SlateCore",
-				
 				"PropertyEditor",
 				"SharedSettingsWidgets",
 				"SourceControl",
@@ -25,6 +24,7 @@ public class MacPlatformEditor : ModuleRules
 				"TargetPlatform",
 				"MaterialShaderQualitySettings",
 				"RenderCore",
+				"SettingsEditor"
 			}
 		);
 
diff --git a/Engine/Source/Developer/Mac/MacPlatformEditor/Private/MacPlatformEditorModule.cpp b/Engine/Source/Developer/Mac/MacPlatformEditor/Private/MacPlatformEditorModule.cpp
--- a/Engine/Source/Developer/Mac/MacPlatformEditor/Private/MacPlatformEditorModule.cpp
+++ b/Engine/Source/Developer/Mac/MacPlatformEditor/Private/MacPlatformEditorModule.cpp
@@ -10,6 +10,7 @@
 #include "UObject/Class.h"
 #include "PropertyEditorModule.h"
 #include "XcodeProjectSettingsDetailsCustomization.h"
+#include "MacTargetSettingsCustomization.h"
 
 #define LOCTEXT_NAMESPACE "MacPlatformEditorModule"
 
@@ -44,6 +45,11 @@ class FMacPlatformEditorModule
                 "XcodeProjectSettings",
                 FOnGetDetailCustomizationInstance::CreateStatic(&FXcodeProjectSettingsDetailsCustomization::MakeInstance)
             );
+			
+			PropertyModule.RegisterCustomClassLayout(
+				"MacTargetSettings",
+				FOnGetDetailCustomizationInstance::CreateStatic(&FMacTargetSettingsCustomization::MakeInstance)
+			);
 		}
 	}
 
diff --git a/Engine/Source/Developer/Mac/MacPlatformEditor/Private/MacTargetSettingsCustomization.cpp b/Engine/Source/Developer/Mac/MacPlatformEditor/Private/MacTargetSettingsCustomization.cpp
new file mode 100644
--- /dev/null
+++ b/Engine/Source/Developer/Mac/MacPlatformEditor/Private/MacTargetSettingsCustomization.cpp
@@ -0,0 +1,83 @@
+// Copyright Epic Games, Inc. All Rights Reserved.
+
+#include "MacTargetSettingsCustomization.h"
+#include "MacTargetSettings.h"
+#include "DetailCategoryBuilder.h"
+#include "DetailLayoutBuilder.h"
+#include "DetailWidgetRow.h"
+#include "Misc/ConfigCacheIni.h"
+#include "Misc/ConfigContext.h"
+#include "Modules/ModuleManager.h"
+#include "ISettingsEditorModule.h"
+#include "SSettingsEditorCheckoutNotice.h"
+#include "Widgets/Input/SCheckBox.h"
+#include "Widgets/Text/STextBlock.h"
+
+#define LOCTEXT_NAMESPACE "MacTargetSettingsCustomization"
+
+TSharedRef<IDetailCustomization> FMacTargetSettingsCustomization::MakeInstance()
+{
+	return MakeShareable(new FMacTargetSettingsCustomization);
+}
+
+void FMacTargetSettingsCustomization::CustomizeDetails(IDetailLayoutBuilder& DetailLayout)
+{
+	IDetailCategoryBuilder& RenderingCategory = DetailLayout.EditCategory(TEXT("Rendering"));
+	RenderingCategory.AddCustomRow(LOCTEXT("UseBindlessWithSM5", "Use Bindless With SM5 (Experimental)"), false)
+	.RowTag("UseBindlessWithSM5")
+	.NameContent()
+	[
+		SNew(STextBlock)
+		.Font(IDetailLayoutBuilder::GetDetailFont())
+		.Text(LOCTEXT("UseBindlessWithSM5", "Use Bindless With SM5 (Experimental)"))
+	]
+	.ValueContent()
+	[
+		SNew(SCheckBox)
+		.IsChecked(this, &FMacTargetSettingsCustomization::ShouldUseBindlessWithSM5)
+		.OnCheckStateChanged(this, &FMacTargetSettingsCustomization::OnUseBindlessWithSM5Changed)
+	];
+}
+
+constexpr TCHAR SM5SPSection[] = TEXT("ShaderPlatformConfig METAL_SM5");
+constexpr TCHAR BindlessConfigurationKey[] = TEXT("BindlessConfiguration");
+
+ECheckBoxState FMacTargetSettingsCustomization::ShouldUseBindlessWithSM5() const
+{
+	FString OutSetting;
+	if (GConfig->GetString(SM5SPSection, BindlessConfigurationKey, OutSetting, GEngineIni))
+	{
+		return OutSetting.Equals(TEXT("Disabled"), ESearchCase::IgnoreCase) ? ECheckBoxState::Unchecked : ECheckBoxState::Checked;
+	}
+	else
+	{
+		return ECheckBoxState::Unchecked;
+	}
+}
+void FMacTargetSettingsCustomization::OnUseBindlessWithSM5Changed(ECheckBoxState NewState)
+{
+	static const FString ClassConfigName = UMacTargetSettings::StaticClass()->ClassConfigName.ToString();
+	static const FString DefaultConfigLayerFilename = UMacTargetSettings::StaticClass()->GetDefaultConfigFilename();
+
+	if (!SettingsHelpers::CheckOutOrAddFile(DefaultConfigLayerFilename))
+	{
+		SettingsHelpers::MakeWritable(DefaultConfigLayerFilename);
+	}
+
+	// CF UObject::UpdateSingleSectionOfConfigFile implementation
+	GConfig->Flush(false, *ClassConfigName);
+	
+	// Because BindlessConfiguration is not a UPROPERTY with config keyboard we have to do it this way.
+	FConfigFile Temp;
+	Temp.SetString(SM5SPSection, BindlessConfigurationKey, NewState == ECheckBoxState::Unchecked ? TEXT("Disabled") : TEXT("All"));
+	Temp.UpdateSinglePropertyInSection(*DefaultConfigLayerFilename, BindlessConfigurationKey, SM5SPSection);
+	
+	// CF UObject::UpdateSingleSectionOfConfigFile implementation
+	FConfigContext Context = FConfigContext::ForceReloadIntoGConfig();
+	Context.bWriteDestIni = false;
+	Context.Load(*ClassConfigName);
+	
+	FModuleManager::GetModuleChecked<ISettingsEditorModule>("SettingsEditor").OnApplicationRestartRequired();
+}
+
+#undef LOCTEXT_NAMESPACE
diff --git a/Engine/Source/Developer/Mac/MacPlatformEditor/Private/MacTargetSettingsCustomization.h b/Engine/Source/Developer/Mac/MacPlatformEditor/Private/MacTargetSettingsCustomization.h
new file mode 100644
--- /dev/null
+++ b/Engine/Source/Developer/Mac/MacPlatformEditor/Private/MacTargetSettingsCustomization.h
@@ -0,0 +1,22 @@
+// Copyright Epic Games, Inc. All Rights Reserved.
+
+#pragma once
+
+#include "CoreMinimal.h"
+#include "IDetailCustomization.h"
+#include "Styling/SlateTypes.h"
+
+class FMacTargetSettingsCustomization : public IDetailCustomization
+{
+public:
+	// Makes a new instance of this detail layout class for a specific detail view requesting it
+	static TSharedRef<IDetailCustomization> MakeInstance();
+
+	// IDetailCustomization interface
+	virtual void CustomizeDetails(IDetailLayoutBuilder& DetailLayout) override;
+	// End of IDetailCustomization interface
+	
+private:
+	ECheckBoxState ShouldUseBindlessWithSM5() const;
+	void OnUseBindlessWithSM5Changed(ECheckBoxState NewState);
+};
diff --git a/Engine/Source/Runtime/Apple/MetalRHI/Private/MetalRHI.cpp b/Engine/Source/Runtime/Apple/MetalRHI/Private/MetalRHI.cpp
--- a/Engine/Source/Runtime/Apple/MetalRHI/Private/MetalRHI.cpp
+++ b/Engine/Source/Runtime/Apple/MetalRHI/Private/MetalRHI.cpp
@@ -486,8 +486,6 @@ FMetalDynamicRHI::FMetalDynamicRHI(ERHIFeatureLevel::Type RequestedFeatureLevel)
 		GMaxRHIFeatureLevel = ERHIFeatureLevel::SM5;
 		GMaxRHIShaderPlatform = SP_METAL_SM5;
 	}
-
-    GRHIGlobals.bSupportsBindless = FDataDrivenShaderPlatformInfo::GetSupportsBindless(GMaxRHIShaderPlatform);
 	
 #if PLATFORM_SUPPORTS_MESH_SHADERS
     GRHISupportsMeshShadersTier0 = RHISupportsMeshShadersTier0(GMaxRHIShaderPlatform);
@@ -506,6 +504,8 @@ FMetalDynamicRHI::FMetalDynamicRHI(ERHIFeatureLevel::Type RequestedFeatureLevel)
 			GMaxRHIShaderPlatform = SP_METAL_ES3_1;
 		}
 	}
+	
+	GRHIGlobals.bSupportsBindless = RHIGetRuntimeBindlessConfiguration(GMaxRHIShaderPlatform) != ERHIBindlessConfiguration::Disabled;
 
 	// Bindless is technically unlimited so we set 32 as Max UAV's, < SM5 8 
 	GRHIGlobals.MaxSimultaneousUAVs = GMaxRHIFeatureLevel >= ERHIFeatureLevel::SM6 ? 32 : 8;
diff --git a/Engine/Source/Runtime/Apple/MetalRHI/Private/MetalRHIPrivate.h b/Engine/Source/Runtime/Apple/MetalRHI/Private/MetalRHIPrivate.h
--- a/Engine/Source/Runtime/Apple/MetalRHI/Private/MetalRHIPrivate.h
+++ b/Engine/Source/Runtime/Apple/MetalRHI/Private/MetalRHIPrivate.h
@@ -146,7 +146,7 @@ extern bool GMetalCommandBufferDebuggingEnabled;
 
 FORCEINLINE bool IsMetalBindlessEnabled()
 {
-	return GRHIGlobals.bSupportsBindless && GMaxRHIFeatureLevel >= ERHIFeatureLevel::SM6;
+	return GRHIGlobals.bSupportsBindless;
 }
 
 // Access the underlying surface object from any kind of texture
diff --git a/Engine/Source/Runtime/RHI/Private/Apple/AppleDynamicRHI.cpp b/Engine/Source/Runtime/RHI/Private/Apple/AppleDynamicRHI.cpp
--- a/Engine/Source/Runtime/RHI/Private/Apple/AppleDynamicRHI.cpp
+++ b/Engine/Source/Runtime/RHI/Private/Apple/AppleDynamicRHI.cpp
@@ -166,7 +166,11 @@ FDynamicRHI* PlatformCreateDynamicRHI()
 		// The Feature Level is known after creating the RHI
 		FString FeatureLevelName;
 		GetFeatureLevelName(GMaxRHIFeatureLevel, FeatureLevelName);
-		FApp::SetGraphicsRHI(FString::Printf(TEXT("Metal (%s)"), *FeatureLevelName));
+		FString BindlessStatus = TEXT(" ");
+#if PLATFORM_SUPPORTS_BINDLESS_RENDERING
+		BindlessStatus = GRHIGlobals.bSupportsBindless ? TEXT(" Bindless ") : TEXT(" Bindful ");
+#endif
+		FApp::SetGraphicsRHI(FString::Printf(TEXT("Metal%s(%s)"), *BindlessStatus, *FeatureLevelName));
 	}
 
 	return DynamicRHI;
diff --git a/Engine/Source/Runtime/RenderCore/Public/UpdateDescriptorHandle.h b/Engine/Source/Runtime/RenderCore/Public/UpdateDescriptorHandle.h
--- a/Engine/Source/Runtime/RenderCore/Public/UpdateDescriptorHandle.h
+++ b/Engine/Source/Runtime/RenderCore/Public/UpdateDescriptorHandle.h
@@ -29,7 +29,7 @@ public:
     
     static bool ShouldCompilePermutation(const FGlobalShaderPermutationParameters& Parameters)
     {
-		return IsMetalSM6Platform(Parameters.Platform);;
+		return IsMetalPlatform(Parameters.Platform) && FDataDrivenShaderPlatformInfo::GetSupportsBindless(Parameters.Platform);
     }
     
     static inline void ModifyCompilationEnvironment(const FGlobalShaderPermutationParameters& Parameters, FShaderCompilerEnvironment& OutEnvironment)
