[MacTargetPlatformControls]
* Moved sdk includes to cpp

#Synced-CL 50230134

diff --git a/Engine/Source/Developer/Mac/MacTargetPlatformControls/Private/MacTargetPlatformControlsModule.cpp b/Engine/Source/Developer/Mac/MacTargetPlatformControls/Private/MacTargetPlatformControlsModule.cpp
--- a/Engine/Source/Developer/Mac/MacTargetPlatformControls/Private/MacTargetPlatformControlsModule.cpp
+++ b/Engine/Source/Developer/Mac/MacTargetPlatformControls/Private/MacTargetPlatformControlsModule.cpp
@@ -2,10 +2,16 @@
 
 #include "CoreMinimal.h"
 #include "GenericMacTargetPlatformControls.h"
+#include "LocalMacTargetDevice.h"
 #include "Modules/ModuleManager.h"
 #include "Interfaces/ITargetPlatformControlsModule.h"
 #include "Modules/ModuleManager.h"
 #include "IMacTargetPlatformSettingsModule.h"
+
+#if PLATFORM_MAC
+#include "Mac/MacSystemIncludes.h"
+#endif
+
 #if WITH_ENGINE
 #include "CookedEditorTargetPlatformControls.h"
 #endif
@@ -66,6 +72,16 @@ public:
 
 };
 
+bool FLocalMacTargetDevice::Reboot( bool bReconnect )
+{
+#if PLATFORM_MAC
+	NSAppleScript* Script = [[NSAppleScript alloc] initWithSource:@"tell application \"System Events\" to restart"];
+	NSDictionary* ErrorDict = [NSDictionary dictionary];
+	[Script executeAndReturnError: &ErrorDict];
+#endif
+	return true;
+}
+
 
 #undef LOCTEXT_NAMESPACE
 
diff --git a/Engine/Source/Developer/Mac/MacTargetPlatformControls/Public/LocalMacTargetDevice.h b/Engine/Source/Developer/Mac/MacTargetPlatformControls/Public/LocalMacTargetDevice.h
--- a/Engine/Source/Developer/Mac/MacTargetPlatformControls/Public/LocalMacTargetDevice.h
+++ b/Engine/Source/Developer/Mac/MacTargetPlatformControls/Public/LocalMacTargetDevice.h
@@ -12,10 +12,6 @@
 #include "HAL/FileManager.h"
 #include "HAL/PlatformProcess.h"
 
-#if PLATFORM_MAC
-#include "Mac/MacSystemIncludes.h"
-#endif
-
 /**
  * Template for local Mac target devices.
  */
@@ -99,15 +95,7 @@ public:
 		return false;
 	}
 
-	virtual bool Reboot( bool bReconnect = false ) override
-	{
-#if PLATFORM_MAC
-		NSAppleScript* Script = [[NSAppleScript alloc] initWithSource:@"tell application \"System Events\" to restart"];
-		NSDictionary* ErrorDict = [NSDictionary dictionary];
-		[Script executeAndReturnError: &ErrorDict];
-#endif
-		return true;
-	}
+	MACTARGETPLATFORMCONTROLS_API virtual bool Reboot( bool bReconnect = false ) override;
 
 	virtual bool SupportsFeature( ETargetDeviceFeatures Feature ) const override
 	{
