[Backout] - CL50365724
#rnx
Original CL Desc
-----------------------------------------------------------------
Adding an Apple Certificate checker to UBT so we can check that designated certificates are being used to sign Apple platform builds

#Synced-CL 50166977

diff --git a/Engine/Source/Developer/Mac/MacTargetPlatformSettings/Classes/XcodeProjectSettings.h b/Engine/Source/Developer/Mac/MacTargetPlatformSettings/Classes/XcodeProjectSettings.h
--- a/Engine/Source/Developer/Mac/MacTargetPlatformSettings/Classes/XcodeProjectSettings.h
+++ b/Engine/Source/Developer/Mac/MacTargetPlatformSettings/Classes/XcodeProjectSettings.h
@@ -168,24 +168,6 @@ public:
 	UPROPERTY(EditAnywhere, config, Category="Code Signing", meta = (EditCondition="bUseAutomaticCodeSigning && bUseAppStoreConnect", DisplayName = "App Store Connect: Key File"))
 	FFilePath AppStoreConnectKeyPath;
 	
-	/**
-	 * SHA1 hashes of certificates to verify. If these certs are not present, the build will either try to install them or fail.
-	 */
-	UPROPERTY(EditAnywhere, config, Category="Code Signing")
-	TArray<FString> CertificateHashes;
-	
-	/**
-	 * Perforce file paths to certificate files to verify. If these certs are not present, the build will either try to install them or fail.
-	 */
-	UPROPERTY(EditAnywhere, config, Category="Code Signing")
-	TArray<FString> P4CertificateFilePaths;
-	
-	/**
-	 * Perforce stream names containing the certificate files to verify. If these certs are not present, the build will either try to install them or fail.
-	 */
-	UPROPERTY(EditAnywhere, config, Category="Code Signing")
-	TArray<FString> CertificateStreamNames;
-	
 	/**
 	 * The path to the optional PrivacyInfo.xcprivacy file for your project
 	 */
diff --git a/Engine/Source/Programs/UnrealBuildTool/Platform/Apple/AppleCertificateVerification.cs b/Engine/Source/Programs/UnrealBuildTool/Platform/Apple/AppleCertificateVerification.cs
deleted file mode 100644
diff --git a/Engine/Source/Programs/UnrealBuildTool/Platform/Apple/UEBuildApple.cs b/Engine/Source/Programs/UnrealBuildTool/Platform/Apple/UEBuildApple.cs
--- a/Engine/Source/Programs/UnrealBuildTool/Platform/Apple/UEBuildApple.cs
+++ b/Engine/Source/Programs/UnrealBuildTool/Platform/Apple/UEBuildApple.cs
@@ -127,10 +127,6 @@ namespace UnrealBuildTool
 		public AppleBuildPlatform(UnrealTargetPlatform Platform, UEBuildPlatformSDK SDK, UnrealArchitectureConfig ArchitectureConfig, ILogger Logger)
 			: base(Platform, SDK, ArchitectureConfig, Logger)
 		{
-			if(BuildHostPlatform.Current.Platform == UnrealTargetPlatform.Mac)
-			{
-				(new AppleCertificateVerification()).VerifyCertificateInstallation(Logger);
-			}
 		}
 
 		public abstract AppleTargetRules GetAppleTargetRules(TargetRules Target);
diff --git a/Engine/Source/Programs/UnrealBuildTool/Platform/IOS/IOSCodeSigning.cs b/Engine/Source/Programs/UnrealBuildTool/Platform/IOS/IOSCodeSigning.cs
--- a/Engine/Source/Programs/UnrealBuildTool/Platform/IOS/IOSCodeSigning.cs
+++ b/Engine/Source/Programs/UnrealBuildTool/Platform/IOS/IOSCodeSigning.cs
@@ -85,9 +85,9 @@ namespace UnrealBuildTool
 		}
 
 		/// <summary>
-		/// Finds all installed certificates as a X509Certificate2Collection
+		/// Finds all valid installed certificates
 		/// </summary>
-		public static X509Certificate2Collection FindCertificateCollection()
+		public static List<string> FindCertificates()
 		{
 			string[] ValidCertificatePrefixes = { "iPhone Developer", "iPhone Distribution", "Apple Development", "Apple Distribution" };
 
@@ -141,16 +141,6 @@ namespace UnrealBuildTool
 				Store.Close();
 			}
 
-			return FoundCerts;
-		}
-
-		/// <summary>
-		/// Finds all valid installed certificates
-		/// </summary>
-		public static List<string> FindCertificates()
-		{
-			X509Certificate2Collection FoundCerts = FindCertificateCollection();
-			
 			List<string> Certs = new();
 			foreach (X509Certificate2 TestCert in FoundCerts)
 			{
