ShaderPlatformConfig editor
- Adding editable ShaderPlatformConfig fields to each platform's project settings.
- Only adding Bindless Configuration for editing.
- To reduce noise with only the single setting that's only used on a few platforms, the settings will be hidden for non-bindless platforms until this is expanded.
- Key objective here is to add the injection point for future settings.

#Synced-CL 51242114

diff --git a/Engine/Source/Developer/Android/AndroidPlatformEditor/AndroidPlatformEditor.Build.cs b/Engine/Source/Developer/Android/AndroidPlatformEditor/AndroidPlatformEditor.Build.cs
--- a/Engine/Source/Developer/Android/AndroidPlatformEditor/AndroidPlatformEditor.Build.cs
+++ b/Engine/Source/Developer/Android/AndroidPlatformEditor/AndroidPlatformEditor.Build.cs
@@ -29,8 +29,9 @@ public class AndroidPlatformEditor : ModuleRules
                 "RHI",
                 "MaterialShaderQualitySettings",
 				"MainFrame",
-                "AudioSettingsEditor"
-            }
+                "AudioSettingsEditor",
+				"ShaderPlatformConfigEditor",
+			}
 		);
 
 		PrivateIncludePathModuleNames.AddRange(
diff --git a/Engine/Source/Developer/Android/AndroidPlatformEditor/Private/AndroidTargetSettingsCustomization.cpp b/Engine/Source/Developer/Android/AndroidPlatformEditor/Private/AndroidTargetSettingsCustomization.cpp
--- a/Engine/Source/Developer/Android/AndroidPlatformEditor/Private/AndroidTargetSettingsCustomization.cpp
+++ b/Engine/Source/Developer/Android/AndroidPlatformEditor/Private/AndroidTargetSettingsCustomization.cpp
@@ -36,6 +36,7 @@
 #include "AndroidLicenseDialog.h"
 #include "Interfaces/IMainFrameModule.h"
 #include "Framework/Application/SlateApplication.h"
+#include "ShaderPlatformConfigEditor.h"
 
 #define LOCTEXT_NAMESPACE "AndroidRuntimeSettings"
 
@@ -83,6 +84,8 @@ void FAndroidTargetSettingsCustomization::CustomizeDetails(IDetailLayoutBuilder&
 	BuildIconSection(DetailLayout);
 	BuildLaunchImageSection(DetailLayout);
 	AudioPluginWidgetManager.BuildAudioCategory(DetailLayout, FString(TEXT("Android")));
+
+	UE::ShaderPlatformConfigEditor::AddShaderPlatformConfigsFromSettingsModule(DetailLayout, "AndroidTargetPlatformSettings");
 }
 
 static void OnBrowserLinkClicked(const FSlateHyperlinkRun::FMetadata& Metadata)
diff --git a/Engine/Source/Developer/IOS/IOSPlatformEditor/IOSPlatformEditor.Build.cs b/Engine/Source/Developer/IOS/IOSPlatformEditor/IOSPlatformEditor.Build.cs
--- a/Engine/Source/Developer/IOS/IOSPlatformEditor/IOSPlatformEditor.Build.cs
+++ b/Engine/Source/Developer/IOS/IOSPlatformEditor/IOSPlatformEditor.Build.cs
@@ -31,6 +31,7 @@ public class IOSPlatformEditor : ModuleRules
 				"FreeImage",
                 "MacTargetPlatformSettings",
 				"MacTargetPlatformControls",
+				"ShaderPlatformConfigEditor",
 			}
 		);
 
diff --git a/Engine/Source/Developer/IOS/IOSPlatformEditor/Private/IOSTargetSettingsCustomization.cpp b/Engine/Source/Developer/IOS/IOSPlatformEditor/Private/IOSTargetSettingsCustomization.cpp
--- a/Engine/Source/Developer/IOS/IOSPlatformEditor/Private/IOSTargetSettingsCustomization.cpp
+++ b/Engine/Source/Developer/IOS/IOSPlatformEditor/Private/IOSTargetSettingsCustomization.cpp
@@ -39,6 +39,7 @@
 #include "HAL/PlatformFileManager.h"
 #include "UEFreeImage.h"
 #include "XcodeProjectSettings.h"
+#include "ShaderPlatformConfigEditor.h"
 
 #define LOCTEXT_NAMESPACE "IOSTargetSettings"
 DEFINE_LOG_CATEGORY_STATIC(LogIOSTargetSettings, Log, All);
@@ -155,6 +156,7 @@ void FIOSTargetSettingsCustomization::CustomizeDetails(IDetailLayoutBuilder& Det
 
 	AudioPluginWidgetManager.BuildAudioCategory(DetailLayout, FString(TEXT("IOS")));
 
+	UE::ShaderPlatformConfigEditor::AddShaderPlatformConfigsFromSettingsModule(DetailLayout, "IOSTargetPlatformSettings");
 }
 
 void FIOSTargetSettingsCustomization::OnSelect()
diff --git a/Engine/Source/Developer/Mac/MacPlatformEditor/MacPlatformEditor.Build.cs b/Engine/Source/Developer/Mac/MacPlatformEditor/MacPlatformEditor.Build.cs
--- a/Engine/Source/Developer/Mac/MacPlatformEditor/MacPlatformEditor.Build.cs
+++ b/Engine/Source/Developer/Mac/MacPlatformEditor/MacPlatformEditor.Build.cs
@@ -27,7 +27,8 @@ public class MacPlatformEditor : ModuleRules
 				"SettingsEditor",
 				"AudioSettingsEditor",
 				"RHI",
-				"UnrealEd"
+				"UnrealEd",
+				"ShaderPlatformConfigEditor",
 			}
 		);
 
diff --git a/Engine/Source/Developer/Mac/MacPlatformEditor/Private/MacTargetSettingsCustomization.cpp b/Engine/Source/Developer/Mac/MacPlatformEditor/Private/MacTargetSettingsCustomization.cpp
--- a/Engine/Source/Developer/Mac/MacPlatformEditor/Private/MacTargetSettingsCustomization.cpp
+++ b/Engine/Source/Developer/Mac/MacPlatformEditor/Private/MacTargetSettingsCustomization.cpp
@@ -30,6 +30,7 @@
 #include "ShaderFormatsPropertyDetails.h"
 #include "ISettingsEditorModule.h"
 #include "SSettingsEditorCheckoutNotice.h"
+#include "ShaderPlatformConfigEditor.h"
 
 namespace MacTargetSettingsDetailsConstants
 {
@@ -174,10 +175,13 @@ void FMacTargetSettingsCustomization::CustomizeDetails( IDetailLayoutBuilder& De
 	FSimpleDelegate OnUpdateShaderStandardWarning = FSimpleDelegate::CreateSP(this, &FMacTargetSettingsCustomization::UpdateShaderStandardWarning);
 	
 	ITargetPlatformSettings* TargetPlatformSettings = FModuleManager::GetModuleChecked<ITargetPlatformSettingsModule>("MacTargetPlatformSettings").GetTargetPlatformSettings()[0];
+
 	// Setup the supported/targeted RHI property view
 	TargetShaderFormatsDetails = MakeShareable(new FShaderFormatsPropertyDetails(&DetailBuilder, TEXT("TargetedRHIs"), TEXT("Targeted RHIs")));
 	TargetShaderFormatsDetails->SetOnUpdateShaderWarning(OnUpdateShaderStandardWarning);
-	TargetShaderFormatsDetails->CreateTargetShaderFormatsPropertyView(TargetPlatformSettings, &GetFriendlyNameFromRHINameMac);
+	TargetShaderFormatsDetails->CreateTargetShaderFormatsPropertyView(TargetPlatformSettings, &GetFriendlyNameFromRHINameMac, nullptr, ECategoryPriority::Important);
+
+	UE::ShaderPlatformConfigEditor::AddShaderPlatformConfigs(DetailBuilder, TargetPlatformSettings);
 	
 	// Setup the shader version property view
 	// Handle max. shader version a little specially.
@@ -220,21 +224,6 @@ void FMacTargetSettingsCustomization::CustomizeDetails( IDetailLayoutBuilder& De
 			]
 		];
 		
-		RenderCategory.AddCustomRow(LOCTEXT("UseBindlessWithSM5", "Use Bindless With SM5 (Experimental)"), false)
-		.RowTag("UseBindlessWithSM5")
-		.NameContent()
-		[
-			SNew(STextBlock)
-			.Font(IDetailLayoutBuilder::GetDetailFont())
-			.Text(LOCTEXT("UseBindlessWithSM5", "Use Bindless With SM5 (Experimental)"))
-		]
-		.ValueContent()
-		[
-			SNew(SCheckBox)
-			.IsChecked(this, &FMacTargetSettingsCustomization::ShouldUseBindlessWithSM5)
-			.OnCheckStateChanged(this, &FMacTargetSettingsCustomization::OnUseBindlessWithSM5Changed)
-		];
-		
 		UpdateShaderStandardWarning();
 	}
 	
@@ -438,45 +427,4 @@ void FMacTargetSettingsCustomization::UpdateShaderStandardWarning()
 	}
 }
 
-constexpr TCHAR SM5SPSection[] = TEXT("ShaderPlatformConfig METAL_SM5");
-constexpr TCHAR BindlessConfigurationKey[] = TEXT("BindlessConfiguration");
-
-ECheckBoxState FMacTargetSettingsCustomization::ShouldUseBindlessWithSM5() const
-{
-	FString OutSetting;
-	if (GConfig->GetString(SM5SPSection, BindlessConfigurationKey, OutSetting, GEngineIni))
-	{
-		return OutSetting.Equals(TEXT("Disabled"), ESearchCase::IgnoreCase) ? ECheckBoxState::Unchecked : ECheckBoxState::Checked;
-	}
-	else
-	{
-		return ECheckBoxState::Unchecked;
-	}
-}
-void FMacTargetSettingsCustomization::OnUseBindlessWithSM5Changed(ECheckBoxState NewState)
-{
-	static const FString ClassConfigName = UMacTargetSettings::StaticClass()->ClassConfigName.ToString();
-	static const FString DefaultConfigLayerFilename = UMacTargetSettings::StaticClass()->GetDefaultConfigFilename();
-
-	if (!SettingsHelpers::CheckOutOrAddFile(DefaultConfigLayerFilename))
-	{
-		SettingsHelpers::MakeWritable(DefaultConfigLayerFilename);
-	}
-
-	// CF UObject::UpdateSingleSectionOfConfigFile implementation
-	GConfig->Flush(false, *ClassConfigName);
-	
-	// Because BindlessConfiguration is not a UPROPERTY with config keyboard we have to do it this way.
-	FConfigFile Temp;
-	Temp.SetString(SM5SPSection, BindlessConfigurationKey, NewState == ECheckBoxState::Unchecked ? TEXT("Disabled") : TEXT("All"));
-	Temp.UpdateSinglePropertyInSection(*DefaultConfigLayerFilename, BindlessConfigurationKey, SM5SPSection);
-	
-	// CF UObject::UpdateSingleSectionOfConfigFile implementation
-	FConfigContext Context = FConfigContext::ForceReloadIntoGConfig();
-	Context.bWriteDestIni = false;
-	Context.Load(*ClassConfigName);
-	
-	FModuleManager::GetModuleChecked<ISettingsEditorModule>("SettingsEditor").OnApplicationRestartRequired();
-}
-
 #undef LOCTEXT_NAMESPACE
diff --git a/Engine/Source/Developer/Mac/MacPlatformEditor/Private/MacTargetSettingsCustomization.h b/Engine/Source/Developer/Mac/MacPlatformEditor/Private/MacTargetSettingsCustomization.h
--- a/Engine/Source/Developer/Mac/MacPlatformEditor/Private/MacTargetSettingsCustomization.h
+++ b/Engine/Source/Developer/Mac/MacPlatformEditor/Private/MacTargetSettingsCustomization.h
@@ -46,12 +46,6 @@ private:
 	/** Delegate to update the shader standard warning */
 	void UpdateShaderStandardWarning();
 	
-	/** Delegate handler to know if we currently use bindless with SM5 */
-	ECheckBoxState ShouldUseBindlessWithSM5() const;
-	
-	/** Delegate handler to set if we use bindless with SM5 */
-	void OnUseBindlessWithSM5Changed(ECheckBoxState NewState);
-	
 private:
 	/** Reference to the target shader formats property view */
 	TSharedPtr<FShaderFormatsPropertyDetails> TargetShaderFormatsDetails;
diff --git a/Engine/Source/Developer/ShaderPlatformConfigEditor/Internal/ShaderPlatformConfigEditor.h b/Engine/Source/Developer/ShaderPlatformConfigEditor/Internal/ShaderPlatformConfigEditor.h
new file mode 100644
--- /dev/null
+++ b/Engine/Source/Developer/ShaderPlatformConfigEditor/Internal/ShaderPlatformConfigEditor.h
@@ -0,0 +1,21 @@
+// Copyright Epic Games, Inc. All Rights Reserved.
+
+#pragma once
+
+#include "HAL/Platform.h"
+
+class FName;
+class IDetailLayoutBuilder;
+class ITargetPlatformSettings;
+
+enum EShaderPlatform : uint16;
+
+namespace UE::ShaderPlatformConfigEditor
+{
+	SHADERPLATFORMCONFIGEDITOR_API void AddShaderPlatformConfig(IDetailLayoutBuilder& DetailBuilder, EShaderPlatform ShaderPlatform);
+	SHADERPLATFORMCONFIGEDITOR_API void AddShaderPlatformConfig(IDetailLayoutBuilder& DetailBuilder, FName ShaderFormat);
+
+	// Shortcuts for adding all shader formats/platforms based on ITargetPlatformSettings
+	SHADERPLATFORMCONFIGEDITOR_API void AddShaderPlatformConfigs(IDetailLayoutBuilder& DetailBuilder, const ITargetPlatformSettings* TargetPlatformSettings);
+	SHADERPLATFORMCONFIGEDITOR_API void AddShaderPlatformConfigsFromSettingsModule(IDetailLayoutBuilder& DetailBuilder, FName SettingsModuleName);
+}
diff --git a/Engine/Source/Developer/ShaderPlatformConfigEditor/Private/ShaderPlatformConfigEditor.cpp b/Engine/Source/Developer/ShaderPlatformConfigEditor/Private/ShaderPlatformConfigEditor.cpp
new file mode 100644
--- /dev/null
+++ b/Engine/Source/Developer/ShaderPlatformConfigEditor/Private/ShaderPlatformConfigEditor.cpp
@@ -0,0 +1,94 @@
+// Copyright Epic Games, Inc. All Rights Reserved.
+
+#include "ShaderPlatformConfigEditor.h"
+
+#include "DataDrivenShaderPlatformInfo.h"
+#include "DetailCategoryBuilder.h"
+#include "DetailLayoutBuilder.h"
+#include "Interfaces/ITargetPlatformSettingsModule.h"
+#include "Modules/ModuleManager.h"
+#include "RHIShaderPlatformConfig.h"
+#include "RHIStrings.h"
+#include "ShaderPlatformConfigSettings.h"
+
+#define LOCTEXT_NAMESPACE "ShaderPlatformConfig"
+
+static UShaderPlatformConfigSettings* GetOrCreateShaderPlatformConfigSettings(EShaderPlatform InShaderPlatform)
+{
+	FStringBuilderBase ConfigSection;
+	UE::RHIShaderPlatformConfig::GetConfigSectionNameForShaderPlatform(ConfigSection, InShaderPlatform);
+
+	FString SettingsName(ConfigSection.ToString());
+	SettingsName.ReplaceInline(TEXT(" "), TEXT("_"));
+
+	UClass* SettingsClass = UShaderPlatformConfigSettings::StaticClass();
+	if (UShaderPlatformConfigSettings* ExistingSettings = FindObject<UShaderPlatformConfigSettings>(SettingsClass, SettingsName))
+	{
+		return ExistingSettings;
+	}
+
+	UShaderPlatformConfigSettings* NewSettings = NewObject<UShaderPlatformConfigSettings>(SettingsClass, FName(SettingsName));
+
+	NewSettings->SetConfiguration(InShaderPlatform, ConfigSection.ToString());
+
+	return NewSettings;
+}
+
+void UE::ShaderPlatformConfigEditor::AddShaderPlatformConfig(IDetailLayoutBuilder& DetailBuilder, EShaderPlatform InShaderPlatform)
+{
+	// ue-todo: figure out how to filter out platforms if they don't have any editable properties.
+	if (!FDataDrivenShaderPlatformInfo::GetSupportsBindless(InShaderPlatform))
+	{
+		return;
+	}
+
+	const FText ShaderPlatformFriendlyName = FDataDrivenShaderPlatformInfo::GetFriendlyName(InShaderPlatform);
+
+	const FText CategoryText = FText::Format(NSLOCTEXT("ShaderPlatformConfig", "CategoryFormat", "Shader Platform Config: {0}"), ShaderPlatformFriendlyName);
+
+	UShaderPlatformConfigSettings* ShaderPlatformConfigSettings = GetOrCreateShaderPlatformConfigSettings(InShaderPlatform);
+
+	IDetailCategoryBuilder& NewCategory = DetailBuilder.EditCategory(ShaderPlatformConfigSettings->GetFName(), CategoryText, ECategoryPriority::Important);
+	FAddPropertyParams Params = FAddPropertyParams()
+		.ForceShowProperty()
+		.HideRootObjectNode(true)
+		.CreateCategoryNodes(false);
+	NewCategory.AddExternalObjects({ ShaderPlatformConfigSettings }, EPropertyLocation::Default, Params);
+}
+
+void UE::ShaderPlatformConfigEditor::AddShaderPlatformConfig(IDetailLayoutBuilder& DetailBuilder, FName ShaderFormat)
+{
+	const EShaderPlatform ShaderPlatform = ShaderFormatToLegacyShaderPlatform(ShaderFormat);
+	if (ShaderPlatform != SP_NumPlatforms)
+	{
+		UE::ShaderPlatformConfigEditor::AddShaderPlatformConfig(DetailBuilder, ShaderPlatform);
+	}
+}
+
+void UE::ShaderPlatformConfigEditor::AddShaderPlatformConfigs(IDetailLayoutBuilder& DetailBuilder, const ITargetPlatformSettings* TargetPlatformSettings)
+{
+	if (TargetPlatformSettings)
+	{
+		TArray<FName> ShaderFormats;
+		TargetPlatformSettings->GetAllPossibleShaderFormats(ShaderFormats);
+
+		for (const FName ShaderFormat : ShaderFormats)
+		{
+			UE::ShaderPlatformConfigEditor::AddShaderPlatformConfig(DetailBuilder, ShaderFormat);
+		}
+	}
+}
+
+void UE::ShaderPlatformConfigEditor::AddShaderPlatformConfigsFromSettingsModule(IDetailLayoutBuilder& DetailBuilder, FName SettingsModuleName)
+{
+	if (ITargetPlatformSettingsModule* Module = FModuleManager::GetModulePtr<ITargetPlatformSettingsModule>(SettingsModuleName))
+	{
+		const TArray<ITargetPlatformSettings*> AllSettings = Module->GetTargetPlatformSettings();
+		for (const ITargetPlatformSettings* TargetPlatformSettings : AllSettings)
+		{
+			UE::ShaderPlatformConfigEditor::AddShaderPlatformConfigs(DetailBuilder, TargetPlatformSettings);
+		}
+	}
+}
+
+#undef LOCTEXT_NAMESPACE
diff --git a/Engine/Source/Developer/ShaderPlatformConfigEditor/Private/ShaderPlatformConfigSettings.cpp b/Engine/Source/Developer/ShaderPlatformConfigEditor/Private/ShaderPlatformConfigSettings.cpp
new file mode 100644
--- /dev/null
+++ b/Engine/Source/Developer/ShaderPlatformConfigEditor/Private/ShaderPlatformConfigSettings.cpp
@@ -0,0 +1,48 @@
+// Copyright Epic Games, Inc. All Rights Reserved.
+
+#include "ShaderPlatformConfigSettings.h"
+
+#include "DataDrivenShaderPlatformInfo.h"
+#include "RHIStrings.h"
+
+UShaderPlatformConfigSettings::UShaderPlatformConfigSettings(const FObjectInitializer& ObjectInitializer)
+	: Super(ObjectInitializer)
+{
+}
+
+const TCHAR* UShaderPlatformConfigSettings::GetConfigOverridePlatform() const
+{
+	return *ConfigPlatform;
+}
+
+void UShaderPlatformConfigSettings::OverridePerObjectConfigSection(FString& InOutSectionName)
+{
+	if (!ConfigSection.IsEmpty())
+	{
+		InOutSectionName = ConfigSection;
+	}
+}
+
+bool UShaderPlatformConfigSettings::CanEditBindlessConfiguration() const
+{
+	return FDataDrivenShaderPlatformInfo::GetSupportsBindless(ShaderPlatform);
+}
+
+TArray<FString> UShaderPlatformConfigSettings::GetBindlessConfigurationOptions() const
+{
+	TArray<FString> Result;
+	Result.Emplace(GetBindlessConfigurationString(ERHIBindlessConfiguration::Disabled));
+	Result.Emplace(GetBindlessConfigurationString(ERHIBindlessConfiguration::RayTracing));
+	Result.Emplace(GetBindlessConfigurationString(ERHIBindlessConfiguration::Minimal));
+	Result.Emplace(GetBindlessConfigurationString(ERHIBindlessConfiguration::All));
+	return Result;
+}
+
+void UShaderPlatformConfigSettings::SetConfiguration(EShaderPlatform InShaderPlatform, const TCHAR* InConfigSection)
+{
+	ShaderPlatform = InShaderPlatform;
+	ConfigPlatform = ShaderPlatformToPlatformName(InShaderPlatform).ToString();
+	ConfigSection = InConfigSection;
+
+	LoadConfig();
+}
diff --git a/Engine/Source/Developer/ShaderPlatformConfigEditor/Private/ShaderPlatformConfigSettings.h b/Engine/Source/Developer/ShaderPlatformConfigEditor/Private/ShaderPlatformConfigSettings.h
new file mode 100644
--- /dev/null
+++ b/Engine/Source/Developer/ShaderPlatformConfigEditor/Private/ShaderPlatformConfigSettings.h
@@ -0,0 +1,36 @@
+// Copyright Epic Games, Inc. All Rights Reserved.
+
+#pragma once
+
+#include "RHIShaderPlatform.h"
+#include "UObject/ObjectMacros.h"
+
+#include "ShaderPlatformConfigSettings.generated.h"
+
+// Temporary solution for getting editing up and running, should be replaced with custom Slate code in the future
+UCLASS(MinimalAPI, config=Engine, DefaultConfig, PerObjectConfig)
+class UShaderPlatformConfigSettings : public UObject
+{
+public:
+	GENERATED_UCLASS_BODY()
+
+	UPROPERTY(config, EditAnywhere, Category="ShaderPlatformConfigSettings", meta = (EditCondition = "CanEditBindlessConfiguration()", GetOptions = "GetBindlessConfigurationOptions"))
+	FString BindlessConfiguration;
+
+	virtual const TCHAR* GetConfigOverridePlatform() const final;
+	virtual void OverridePerObjectConfigSection(FString& InOutSectionName) final;
+
+	// BindlessConfiguration support
+	UFUNCTION()
+	bool CanEditBindlessConfiguration() const;
+
+	UFUNCTION()
+	TArray<FString> GetBindlessConfigurationOptions() const;
+
+	void SetConfiguration(EShaderPlatform InShaderPlatform, const TCHAR* InConfigSection);
+
+private:
+	EShaderPlatform ShaderPlatform{};
+	FString ConfigPlatform;
+	FString ConfigSection;
+};
diff --git a/Engine/Source/Developer/ShaderPlatformConfigEditor/ShaderPlatformConfigEditor.Build.cs b/Engine/Source/Developer/ShaderPlatformConfigEditor/ShaderPlatformConfigEditor.Build.cs
new file mode 100644
--- /dev/null
+++ b/Engine/Source/Developer/ShaderPlatformConfigEditor/ShaderPlatformConfigEditor.Build.cs
@@ -0,0 +1,33 @@
+// Copyright Epic Games, Inc. All Rights Reserved.
+
+using UnrealBuildTool;
+
+public class ShaderPlatformConfigEditor : ModuleRules
+{
+	public ShaderPlatformConfigEditor(ReadOnlyTargetRules Target) : base(Target)
+	{
+		PrivateDependencyModuleNames.AddRange(
+			new string[] {
+				"Core",
+				"InputCore",
+				"Engine",
+				"Slate",
+				"SlateCore",
+				"RHI",
+				"RenderCore",
+
+				"PropertyEditor",
+				"UnrealEd",
+				"CoreUObject",
+				"TargetPlatform",
+			}
+		);
+
+		PublicIncludePathModuleNames.AddRange(
+			new string[] {
+				"Engine",
+				"RHI",
+			}
+		);
+	}
+}
diff --git a/Engine/Source/Developer/Windows/WindowsTargetPlatformSettings/Classes/WindowsTargetSettings.h b/Engine/Source/Developer/Windows/WindowsTargetPlatformSettings/Classes/WindowsTargetSettings.h
--- a/Engine/Source/Developer/Windows/WindowsTargetPlatformSettings/Classes/WindowsTargetSettings.h
+++ b/Engine/Source/Developer/Windows/WindowsTargetPlatformSettings/Classes/WindowsTargetSettings.h
@@ -34,6 +34,12 @@ enum class EDefaultGraphicsRHI : uint8
 	DefaultGraphicsRHI_Vulkan = 3 UMETA(DisplayName = "Vulkan"),
 };
 
+UENUM()
+enum class EBindlessConfigurationSetting
+{
+	Disabled = 0 UMETA(DisplayName = "Disabled")
+};
+
 /**
  * Implements the settings for the Windows target platform. The first instance of this class is initialized in
  * WindowsTargetPlatform, really early during the startup sequence before the CDO has been constructed, so its config 
diff --git a/Engine/Source/Editor/DetailCustomizations/DetailCustomizations.Build.cs b/Engine/Source/Editor/DetailCustomizations/DetailCustomizations.Build.cs
--- a/Engine/Source/Editor/DetailCustomizations/DetailCustomizations.Build.cs
+++ b/Engine/Source/Editor/DetailCustomizations/DetailCustomizations.Build.cs
@@ -73,6 +73,7 @@ public class DetailCustomizations : ModuleRules
 				"Json",
 				"JsonUtilities",
 				"SettingsEditor",
+				"ShaderPlatformConfigEditor",
 			}
 		);
 
diff --git a/Engine/Source/Editor/DetailCustomizations/Private/WindowsTargetSettingsDetails.cpp b/Engine/Source/Editor/DetailCustomizations/Private/WindowsTargetSettingsDetails.cpp
--- a/Engine/Source/Editor/DetailCustomizations/Private/WindowsTargetSettingsDetails.cpp
+++ b/Engine/Source/Editor/DetailCustomizations/Private/WindowsTargetSettingsDetails.cpp
@@ -32,6 +32,7 @@
 #include "UnrealEngine.h"
 #include "RHIShaderFormatDefinitions.inl"
 #include "ShaderFormatsPropertyDetails.h"
+#include "ShaderPlatformConfigEditor.h"
 
 #if WITH_ENGINE
 #include "AudioDevice.h"
@@ -210,6 +211,8 @@ void FWindowsTargetSettingsDetails::CustomizeDetails( IDetailLayoutBuilder& Deta
 	VulkanTargetShaderFormatsDetails = MakeShareable(new FShaderFormatsPropertyDetails(&DetailBuilder, TEXT("VulkanTargetedShaderFormats"), TEXT("Vulkan Targeted Shader Formats")));
 	VulkanTargetShaderFormatsDetails->CreateTargetShaderFormatsPropertyView(TargetPlatformSettings, &GetFriendlyNameForWindowsShaderPlatformCheckbox, &FilterShaderPlatform_Vulkan, ECategoryPriority::Important);
 
+	UE::ShaderPlatformConfigEditor::AddShaderPlatformConfigs(DetailBuilder, TargetPlatformSettings);
+
 	// Next add the splash image customization
 	const FText EditorSplashDesc(LOCTEXT("EditorSplashLabel", "Editor Splash"));
 	IDetailCategoryBuilder& SplashCategoryBuilder = DetailBuilder.EditCategory(TEXT("Splash"));
