fixing asset cooking and loading for correct behavior on mac es 3.1 projects

#Synced-CL 52243953

#Synced-CL 52243953

diff --git a/Engine/Source/Developer/Mac/MacTargetPlatformSettings/MacTargetPlatformSettings.Build.cs b/Engine/Source/Developer/Mac/MacTargetPlatformSettings/MacTargetPlatformSettings.Build.cs
--- a/Engine/Source/Developer/Mac/MacTargetPlatformSettings/MacTargetPlatformSettings.Build.cs
+++ b/Engine/Source/Developer/Mac/MacTargetPlatformSettings/MacTargetPlatformSettings.Build.cs
@@ -32,6 +32,7 @@ public class MacTargetPlatformSettings : ModuleRules
             PublicDependencyModuleNames.AddRange(
                 new string[] {
                     "CookedEditor",
+					"RHI",
                 }
             );
             PrivateDependencyModuleNames.Add("Engine");
diff --git a/Engine/Source/Developer/Mac/MacTargetPlatformSettings/Public/GenericMacTargetPlatformSettings.h b/Engine/Source/Developer/Mac/MacTargetPlatformSettings/Public/GenericMacTargetPlatformSettings.h
--- a/Engine/Source/Developer/Mac/MacTargetPlatformSettings/Public/GenericMacTargetPlatformSettings.h
+++ b/Engine/Source/Developer/Mac/MacTargetPlatformSettings/Public/GenericMacTargetPlatformSettings.h
@@ -11,6 +11,7 @@
 #include "Mac/MacPlatformProperties.h"
 
 #if WITH_ENGINE
+#include "RHI.h"
 #include "TextureResource.h"
 #include "StaticMeshResources.h"
 #endif // WITH_ENGINE
@@ -56,6 +57,33 @@ public:
 		{
 			return HAS_EDITOR_DATA;
 		}
+		
+		if(Feature == ETargetPlatformFeatures::MobileRendering)
+		{
+			static bool bEnableMobileRenderingFeatureOnMac = false;
+
+#if WITH_ENGINE
+			static bool bHasCachedValue = false;
+			if (!bHasCachedValue)
+			{
+				TArray<FName> TargetedShaderFormats;
+				GetAllTargetedShaderFormats(TargetedShaderFormats);
+
+				for (const FName& Format : TargetedShaderFormats)
+				{
+					if (IsMobilePlatform(ShaderFormatToLegacyShaderPlatform(Format)))
+					{
+						bEnableMobileRenderingFeatureOnMac = true;
+						break;
+					}
+				}
+				
+				bHasCachedValue = true;
+			}
+#endif
+			
+			return bEnableMobileRenderingFeatureOnMac;
+		}
 
 		return TSuper::SupportsFeature(Feature);
 	}
